Fixed cooking issue with DatasmithRuntime plugin

diff --git a/Engine/Plugins/Experimental/Enterprise/DatasmithRuntime/Source/Private/DatasmithRuntime.cpp b/Engine/Plugins/Experimental/Enterprise/DatasmithRuntime/Source/Private/DatasmithRuntime.cpp
--- a/Engine/Plugins/Experimental/Enterprise/DatasmithRuntime/Source/Private/DatasmithRuntime.cpp
+++ b/Engine/Plugins/Experimental/Enterprise/DatasmithRuntime/Source/Private/DatasmithRuntime.cpp
@@ -80,7 +80,7 @@ ADatasmithRuntimeActor::ADatasmithRuntimeActor()
 	{
 	}
 
-	RootComponent = NewObject< USceneComponent >( this, USceneComponent::StaticClass(), TEXT("DatasmithRuntimeComponent"), RF_Transactional );
+	RootComponent = CreateDefaultSubobject<USceneComponent>(TEXT("DatasmithRuntimeComponent"));
 	AddInstanceComponent( RootComponent );
 	RootComponent->SetMobility(EComponentMobility::Movable);
 	RootComponent->Bounds = DefaultBounds;
