LWC: miscellaneous type fixes

diff --git a/Engine/Plugins/Experimental/Enterprise/DatasmithRuntime/Source/Private/DatasmithRuntimeUtils.cpp b/Engine/Plugins/Experimental/Enterprise/DatasmithRuntime/Source/Private/DatasmithRuntimeUtils.cpp
--- a/Engine/Plugins/Experimental/Enterprise/DatasmithRuntime/Source/Private/DatasmithRuntimeUtils.cpp
+++ b/Engine/Plugins/Experimental/Enterprise/DatasmithRuntime/Source/Private/DatasmithRuntimeUtils.cpp
@@ -512,7 +512,7 @@ namespace DatasmithRuntime
 		for (int32 Index = 0; Index < DirectionCount; ++Index)
 		{
 			FPoly&	Polygon = Element.AddZeroed_GetRef();
-			FVector Base, AxisX, AxisY;
+			FVector3f Base, AxisX, AxisY;
 
 			Polygon.Init();
 			Polygon.Normal = Planes[Index];
@@ -521,16 +521,16 @@ namespace DatasmithRuntime
 			Base = Planes[Index] * Planes[Index].W;
 
 			Polygon.Vertices.Reserve(4);
-			new(Polygon.Vertices) FVector(Base + AxisX * HalfWorldMax + AxisY * HalfWorldMax);
-			new(Polygon.Vertices) FVector(Base + AxisX * HalfWorldMax - AxisY * HalfWorldMax);
-			new(Polygon.Vertices) FVector(Base - AxisX * HalfWorldMax - AxisY * HalfWorldMax);
-			new(Polygon.Vertices) FVector(Base - AxisX * HalfWorldMax + AxisY * HalfWorldMax);
+			new(Polygon.Vertices) FVector3f(Base + AxisX * HalfWorldMax + AxisY * HalfWorldMax);
+			new(Polygon.Vertices) FVector3f(Base + AxisX * HalfWorldMax - AxisY * HalfWorldMax);
+			new(Polygon.Vertices) FVector3f(Base - AxisX * HalfWorldMax - AxisY * HalfWorldMax);
+			new(Polygon.Vertices) FVector3f(Base - AxisX * HalfWorldMax + AxisY * HalfWorldMax);
 
 			for (int32 Jndex = 0; Jndex < DirectionCount; ++Jndex)
 			{
 				if(Index != Jndex)
 				{
-					if(!Polygon.Split(-FVector(Planes[Jndex]), Planes[Jndex] * Planes[Jndex].W))
+					if(!Polygon.Split(-FVector3f(Planes[Jndex]), Planes[Jndex] * Planes[Jndex].W))
 					{
 						Polygon.Vertices.Empty();
 						break;
@@ -560,7 +560,7 @@ namespace DatasmithRuntime
 
 		for (FPoly& Poly : Element)
 		{
-			for (FVector& Position : Poly.Vertices)
+			for (FVector3f& Position : Poly.Vertices)
 			{
 				ConvexElem.VertexData.Add(Position);
 			}
diff --git a/Engine/Plugins/Experimental/GPULightmass/Source/GPULightmass/Private/LightmapRenderer.cpp b/Engine/Plugins/Experimental/GPULightmass/Source/GPULightmass/Private/LightmapRenderer.cpp
--- a/Engine/Plugins/Experimental/GPULightmass/Source/GPULightmass/Private/LightmapRenderer.cpp
+++ b/Engine/Plugins/Experimental/GPULightmass/Source/GPULightmass/Private/LightmapRenderer.cpp
@@ -931,7 +931,7 @@ void FSceneRenderState::SetupRayTracingScene(int32 LODIndex)
 
 							FRayTracingGeometrySegment Segment;
 							Segment.VertexBuffer = nullptr;
-							Segment.VertexBufferStride = sizeof(FVector);
+							Segment.VertexBufferStride = sizeof(FVector3f);
 							Segment.VertexBufferElementType = VET_Float3;
 							Segment.MaxVertices = NumPrimitives * 3;
 							Segment.NumPrimitives = NumPrimitives;
@@ -967,7 +967,7 @@ void FSceneRenderState::SetupRayTracingScene(int32 LODIndex)
 								MeshBatches,
 								false,
 								(uint32)FMath::Square(Landscape.SubsectionSizeVerts),
-								FMath::Square(Landscape.SubsectionSizeVerts) * (uint32)sizeof(FVector),
+								FMath::Square(Landscape.SubsectionSizeVerts) * (uint32)sizeof(FVector3f),
 								(uint32)FMath::Square(Landscape.SubsectionSizeVerts - 1) * 2,
 								&Landscape.SectionRayTracingStates[SubSectionIdx]->Geometry,
 								&Landscape.SectionRayTracingStates[SubSectionIdx]->RayTracingDynamicVertexBuffer,
diff --git a/Engine/Plugins/Experimental/GeometryCache/Source/GeometryCache/Private/GeometryCacheMeshData.cpp b/Engine/Plugins/Experimental/GeometryCache/Source/GeometryCache/Private/GeometryCacheMeshData.cpp
--- a/Engine/Plugins/Experimental/GeometryCache/Source/GeometryCache/Private/GeometryCacheMeshData.cpp
+++ b/Engine/Plugins/Experimental/GeometryCache/Source/GeometryCache/Private/GeometryCacheMeshData.cpp
@@ -98,7 +98,7 @@ uint64 FGeometryCacheMeshData::GetHash() const
 
 	if (Positions.Num() > 0)
 	{
-		Hash = CityHash64((char*) Positions.GetData(), Positions.Num() * sizeof(FVector));
+		Hash = CityHash64((char*) Positions.GetData(), Positions.Num() * sizeof(FVector3f));
 	}
 
 	if (TextureCoordinates.Num() > 0)
diff --git a/Engine/Plugins/Experimental/GeometryCache/Source/GeometryCache/Private/GeometryCacheSceneProxy.cpp b/Engine/Plugins/Experimental/GeometryCache/Source/GeometryCache/Private/GeometryCacheSceneProxy.cpp
--- a/Engine/Plugins/Experimental/GeometryCache/Source/GeometryCache/Private/GeometryCacheSceneProxy.cpp
+++ b/Engine/Plugins/Experimental/GeometryCache/Source/GeometryCache/Private/GeometryCacheSceneProxy.cpp
@@ -1246,8 +1246,8 @@ void FGeomCacheTrackProxy::InitRenderResources(int32 NumVertices, int32 NumIndic
 	TextureCoordinatesBuffer.Init(NumVertices * sizeof(FVector2D));
 	ColorBuffer.Init(NumVertices * sizeof(FColor));
 
-	PositionBuffers[0].Init(NumVertices * sizeof(FVector));
-	PositionBuffers[1].Init(NumVertices * sizeof(FVector));
+	PositionBuffers[0].Init(NumVertices * sizeof(FVector3f));
+	PositionBuffers[1].Init(NumVertices * sizeof(FVector3f));
 	CurrentPositionBufferIndex = -1;
 	PositionBufferFrameIndices[0] = PositionBufferFrameIndices[1] = -1;
 	PositionBufferFrameTimes[0] = PositionBufferFrameTimes[1] = -1.0f;
diff --git a/Engine/Plugins/Runtime/HairStrands/Source/HairStrandsCore/Private/HairStrandsRendering.cpp b/Engine/Plugins/Runtime/HairStrands/Source/HairStrandsCore/Private/HairStrandsRendering.cpp
--- a/Engine/Plugins/Runtime/HairStrands/Source/HairStrandsCore/Private/HairStrandsRendering.cpp
+++ b/Engine/Plugins/Runtime/HairStrands/Source/HairStrandsCore/Private/HairStrandsRendering.cpp
@@ -198,7 +198,7 @@ static void AddGroomCacheUpdatePass(
 	FRDGBufferRef VertexBuffer = nullptr;
 
 	const uint32 DataCount = GroomCacheData.VertexData.PointsPosition.Num();
-	const uint32 DataSizeInBytes = sizeof(FVector) * DataCount;
+	const uint32 DataSizeInBytes = sizeof(FVector3f) * DataCount;
 	if (DataSizeInBytes != 0)
 	{
 		// Deformation are upload into a Buffer<float> as the original position are float3 which is is both
diff --git a/Engine/Plugins/Runtime/HairStrands/Source/HairStrandsCore/Public/GroomCacheData.h b/Engine/Plugins/Runtime/HairStrands/Source/HairStrandsCore/Public/GroomCacheData.h
--- a/Engine/Plugins/Runtime/HairStrands/Source/HairStrandsCore/Public/GroomCacheData.h
+++ b/Engine/Plugins/Runtime/HairStrands/Source/HairStrandsCore/Public/GroomCacheData.h
@@ -90,7 +90,7 @@ struct FGroomCacheVertexData
 	FGroomCacheVertexData() = default;
 	FGroomCacheVertexData(struct FHairStrandsPoints&& PointsData);
 
-	TArray<FVector> PointsPosition;
+	TArray<FVector3f> PointsPosition;
 	TArray<float> PointsRadius;
 	TArray<float> PointsCoordU;
 	TArray<FLinearColor> PointsBaseColor;
diff --git a/Engine/Source/Runtime/D3D12RHI/Private/D3D12RayTracing.cpp b/Engine/Source/Runtime/D3D12RHI/Private/D3D12RayTracing.cpp
--- a/Engine/Source/Runtime/D3D12RHI/Private/D3D12RayTracing.cpp
+++ b/Engine/Source/Runtime/D3D12RHI/Private/D3D12RayTracing.cpp
@@ -2665,7 +2665,7 @@ void TranslateRayTracingGeometryDescs(const FRayTracingGeometryInitializer& Init
 		}
 		else if (Initializer.GeometryType == RTGT_Procedural)
 		{
-			checkf(Segment.VertexBufferStride >= (2 * sizeof(FVector)), TEXT("Procedural geometry vertex buffer must contain at least 2xFloat3 that defines 3D bounding boxes of primitives."));
+			checkf(Segment.VertexBufferStride >= (2 * sizeof(FVector3f)), TEXT("Procedural geometry vertex buffer must contain at least 2xFloat3 that defines 3D bounding boxes of primitives."));
 		}
 
 		check(Segment.FirstPrimitive + Segment.NumPrimitives <= Initializer.TotalPrimitiveCount);
diff --git a/Engine/Source/Runtime/Engine/Private/InstancedStaticMesh.cpp b/Engine/Source/Runtime/Engine/Private/InstancedStaticMesh.cpp
--- a/Engine/Source/Runtime/Engine/Private/InstancedStaticMesh.cpp
+++ b/Engine/Source/Runtime/Engine/Private/InstancedStaticMesh.cpp
@@ -1665,7 +1665,7 @@ void FInstancedStaticMeshSceneProxy::GetDynamicRayTracingInstances(struct FRayTr
 									SimulationInstance.Materials,
 									false,
 									(uint32)LODModel.GetNumVertices(),
-									uint32((SIZE_T)LODModel.GetNumVertices() * sizeof(FVector)),
+									uint32((SIZE_T)LODModel.GetNumVertices() * sizeof(FVector3f)),
 									DynamicData.DynamicGeometry.Initializer.TotalPrimitiveCount,
 									&DynamicData.DynamicGeometry,
 									nullptr,
@@ -1735,7 +1735,7 @@ void FInstancedStaticMeshSceneProxy::GetDynamicRayTracingInstances(struct FRayTr
 							SimulationInstance.Materials,
 							false,
 							(uint32)LODModel.GetNumVertices(),
-							uint32((SIZE_T)LODModel.GetNumVertices() * sizeof(FVector)),
+							uint32((SIZE_T)LODModel.GetNumVertices() * sizeof(FVector3f)),
 							DynamicData.DynamicGeometry.Initializer.TotalPrimitiveCount,
 							&DynamicData.DynamicGeometry,
 							nullptr,
diff --git a/Engine/Source/Runtime/Landscape/Private/LandscapeRender.cpp b/Engine/Source/Runtime/Landscape/Private/LandscapeRender.cpp
--- a/Engine/Source/Runtime/Landscape/Private/LandscapeRender.cpp
+++ b/Engine/Source/Runtime/Landscape/Private/LandscapeRender.cpp
@@ -2652,7 +2652,7 @@ void FLandscapeComponentSceneProxy::GetDynamicRayTracingInstances(FRayTracingMat
 						RayTracingInstance.Materials,
 						false,
 						(uint32)FMath::Square(LodSubsectionSizeVerts),
-						FMath::Square(LodSubsectionSizeVerts) * (uint32)sizeof(FVector),
+						FMath::Square(LodSubsectionSizeVerts) * (uint32)sizeof(FVector3f),
 						(uint32)FMath::Square(LodSubsectionSizeVerts - 1) * 2,
 						&SectionRayTracingStates[SubSectionIdx].Geometry,
 						&SectionRayTracingStates[SubSectionIdx].RayTracingDynamicVertexBuffer,
diff --git a/Engine/Source/Runtime/Renderer/Private/RayTracing/RayTracingDynamicGeometry.cpp b/Engine/Source/Runtime/Renderer/Private/RayTracing/RayTracingDynamicGeometry.cpp
--- a/Engine/Source/Runtime/Renderer/Private/RayTracing/RayTracingDynamicGeometry.cpp
+++ b/Engine/Source/Runtime/Renderer/Private/RayTracing/RayTracingDynamicGeometry.cpp
@@ -234,7 +234,7 @@ void FRayTracingDynamicGeometryCollection::AddDynamicMeshBatchForGeometryUpdate(
 			NumCPUVertices = 1 + MeshBatch.Elements[0].MaxVertexIndex - MeshBatch.Elements[0].MinVertexIndex;
 		}
 
-		const uint32 VertexBufferNumElements = UpdateParams.VertexBufferSize / sizeof(FVector) - MinVertexIndex;
+		const uint32 VertexBufferNumElements = UpdateParams.VertexBufferSize / sizeof(FVector3f) - MinVertexIndex;
 		if (!ensureMsgf(NumCPUVertices <= VertexBufferNumElements, 
 			TEXT("Vertex buffer contains %d vertices, but RayTracingDynamicGeometryConverterCS dispatch command expects at least %d."),
 			VertexBufferNumElements, NumCPUVertices))
