Rename CADKernel CVar

diff --git a/Engine/Plugins/Enterprise/DatasmithCADImporter/Source/DatasmithCADTranslator/Private/DatasmithCADTranslator.cpp b/Engine/Plugins/Enterprise/DatasmithCADImporter/Source/DatasmithCADTranslator/Private/DatasmithCADTranslator.cpp
--- a/Engine/Plugins/Enterprise/DatasmithCADImporter/Source/DatasmithCADTranslator/Private/DatasmithCADTranslator.cpp
+++ b/Engine/Plugins/Enterprise/DatasmithCADImporter/Source/DatasmithCADTranslator/Private/DatasmithCADTranslator.cpp
@@ -12,13 +12,13 @@
 #include "DatasmithUtils.h"
 
 static TAutoConsoleVariable<int32> CVarStaticCADTranslatorEnableThreadedImport(
-	TEXT("r.CADTranslator.EnableThreadedImport"),
+	TEXT("DS.CADTranslator.EnableThreadedImport"),
 	1,
 	TEXT("Activate to parallelise CAD file processing.\n"),
 	ECVF_Default);
 
 static TAutoConsoleVariable<int32> CVarStaticCADTranslatorEnableKernelIOTessellation(
-	TEXT("r.CADTranslator.EnableKernelIOTessellation"),
+	TEXT("DS.CADTranslator.EnableKernelIOTessellation"),
 	1,
 	TEXT("Activate to use KernelIO Tessellation.\n"),
 	ECVF_Default);
diff --git a/Engine/Plugins/Enterprise/DatasmithCADImporter/Source/DatasmithCADTranslator/Private/DatasmithCADTranslatorModule.cpp b/Engine/Plugins/Enterprise/DatasmithCADImporter/Source/DatasmithCADTranslator/Private/DatasmithCADTranslatorModule.cpp
--- a/Engine/Plugins/Enterprise/DatasmithCADImporter/Source/DatasmithCADTranslator/Private/DatasmithCADTranslatorModule.cpp
+++ b/Engine/Plugins/Enterprise/DatasmithCADImporter/Source/DatasmithCADTranslator/Private/DatasmithCADTranslatorModule.cpp
@@ -26,7 +26,7 @@ void FDatasmithCADTranslatorModule::StartupModule()
 	}
 
 	bool bCreateCacheFolder = false;
-	if (IConsoleVariable* CVar = IConsoleManager::Get().FindConsoleVariable(TEXT("r.CADTranslator.EnableCADCache")))
+	if (IConsoleVariable* CVar = IConsoleManager::Get().FindConsoleVariable(TEXT("DS.CADTranslator.EnableCADCache")))
 	{
 		bCreateCacheFolder = CVar->GetInt() != 0;
 	}
diff --git a/Engine/Plugins/Enterprise/DatasmithCADImporter/Source/DatasmithDispatcher/Private/DatasmithDispatcher.cpp b/Engine/Plugins/Enterprise/DatasmithCADImporter/Source/DatasmithDispatcher/Private/DatasmithDispatcher.cpp
--- a/Engine/Plugins/Enterprise/DatasmithCADImporter/Source/DatasmithDispatcher/Private/DatasmithDispatcher.cpp
+++ b/Engine/Plugins/Enterprise/DatasmithCADImporter/Source/DatasmithDispatcher/Private/DatasmithDispatcher.cpp
@@ -14,14 +14,14 @@
 
 #ifndef DATASMITH_CAD_IGNORE_CACHE
 static TAutoConsoleVariable<int32> CVarStaticCADTranslatorEnableCADCache(
-	TEXT("r.CADTranslator.EnableCADCache"),
+	TEXT("DS.CADTranslator.EnableCADCache"),
 	1,
 	TEXT("Activate to save temporary CAD processing file. These file will be use in a next import to avoid CAD file processing\n"),
 	ECVF_Default);
 #endif
 
 static TAutoConsoleVariable<int32> CVarStaticCADTranslatorEnableTimeControl(
-	TEXT("r.CADTranslator.EnableTimeControl"),
+	TEXT("DS.CADTranslator.EnableTimeControl"),
 	1,
 	TEXT("Enable the timer that kill the worker if the import time is unusually long. With this time control, the load of the corrupted file is canceled but the rest of the scene is imported.\n"),
 	ECVF_Default);
diff --git a/Engine/Plugins/Experimental/Enterprise/DatasmithRuntime/Source/Private/DatasmithRuntime.cpp b/Engine/Plugins/Experimental/Enterprise/DatasmithRuntime/Source/Private/DatasmithRuntime.cpp
--- a/Engine/Plugins/Experimental/Enterprise/DatasmithRuntime/Source/Private/DatasmithRuntime.cpp
+++ b/Engine/Plugins/Experimental/Enterprise/DatasmithRuntime/Source/Private/DatasmithRuntime.cpp
@@ -96,13 +96,13 @@ void ADatasmithRuntimeActor::Tick(float DeltaTime)
 #if WITH_EDITOR
 				if (EnableThreadedImport != MAX_int32)
 				{
-					IConsoleVariable* CVar = IConsoleManager::Get().FindConsoleVariable(TEXT("r.CADTranslator.EnableThreadedImport"));
+					IConsoleVariable* CVar = IConsoleManager::Get().FindConsoleVariable(TEXT("DS.CADTranslator.EnableThreadedImport"));
 					CVar->Set(EnableThreadedImport);
 				}
 
 				if (EnableCADCache != MAX_int32)
 				{
-					IConsoleVariable* CVar = IConsoleManager::Get().FindConsoleVariable(TEXT("r.CADTranslator.EnableCADCache"));
+					IConsoleVariable* CVar = IConsoleManager::Get().FindConsoleVariable(TEXT("DS.CADTranslator.EnableCADCache"));
 					CVar->Set(EnableCADCache);
 				}
 #endif
@@ -363,14 +363,14 @@ bool ADatasmithRuntimeActor::LoadFile(const FString& FilePath)
 
 #if WITH_EDITOR
 	EnableThreadedImport = MAX_int32;
-	if (IConsoleVariable* CVar = IConsoleManager::Get().FindConsoleVariable(TEXT("r.CADTranslator.EnableThreadedImport")))
+	if (IConsoleVariable* CVar = IConsoleManager::Get().FindConsoleVariable(TEXT("DS.CADTranslator.EnableThreadedImport")))
 	{
 		EnableThreadedImport = CVar->GetInt();
 		CVar->Set(0);
 	}
 
 	EnableCADCache = MAX_int32;
-	if (IConsoleVariable* CVar = IConsoleManager::Get().FindConsoleVariable(TEXT("r.CADTranslator.EnableCADCache")))
+	if (IConsoleVariable* CVar = IConsoleManager::Get().FindConsoleVariable(TEXT("DS.CADTranslator.EnableCADCache")))
 	{
 		EnableCADCache = CVar->GetInt();
 		CVar->Set(0);
diff --git a/Engine/Plugins/Experimental/Enterprise/DatasmithRuntime/Source/Private/DatasmithRuntimeModule.cpp b/Engine/Plugins/Experimental/Enterprise/DatasmithRuntime/Source/Private/DatasmithRuntimeModule.cpp
--- a/Engine/Plugins/Experimental/Enterprise/DatasmithRuntime/Source/Private/DatasmithRuntimeModule.cpp
+++ b/Engine/Plugins/Experimental/Enterprise/DatasmithRuntime/Source/Private/DatasmithRuntimeModule.cpp
@@ -92,12 +92,12 @@ public:
 					//FModuleManager::Get().LoadModuleChecked(DATASMITHOPENNURBSTRANSLATOR_MODULE_NAME);
 					FModuleManager::Get().LoadModuleChecked(DATASMITHCADTRANSLATOR_MODULE_NAME);
 
-					if (IConsoleVariable* CVar = IConsoleManager::Get().FindConsoleVariable(TEXT("r.CADTranslator.EnableThreadedImport")))
+					if (IConsoleVariable* CVar = IConsoleManager::Get().FindConsoleVariable(TEXT("DS.CADTranslator.EnableThreadedImport")))
 					{
 						CVar->Set(0);
 					}
 
-					if (IConsoleVariable* CVar = IConsoleManager::Get().FindConsoleVariable(TEXT("r.CADTranslator.EnableCADCache")))
+					if (IConsoleVariable* CVar = IConsoleManager::Get().FindConsoleVariable(TEXT("DS.CADTranslator.EnableCADCache")))
 					{
 						CVar->Set(0);
 					}
