Added defensive code in case runtime actor is used in editor's level instead within a game or PIE
Added helper function to return count of references on a UE asset
Fixed crash with SketchUp updates

diff --git a/Engine/Plugins/Experimental/Enterprise/DatasmithRuntime/Source/Private/DatasmithRuntime.cpp b/Engine/Plugins/Experimental/Enterprise/DatasmithRuntime/Source/Private/DatasmithRuntime.cpp
--- a/Engine/Plugins/Experimental/Enterprise/DatasmithRuntime/Source/Private/DatasmithRuntime.cpp
+++ b/Engine/Plugins/Experimental/Enterprise/DatasmithRuntime/Source/Private/DatasmithRuntime.cpp
@@ -209,19 +209,19 @@ void ADatasmithRuntimeActor::OnChangedElement(DirectLink::FSceneGraphId ElementI
 
 bool ADatasmithRuntimeActor::IsConnected()
 {
-	return DirectLinkHelper->IsConnected();
+	return DirectLinkHelper.IsValid() ? DirectLinkHelper->IsConnected() : false;
 }
 
 FString ADatasmithRuntimeActor::GetSourceName()
 {
-	return DirectLinkHelper->GetSourceName();
+	return DirectLinkHelper.IsValid() ? DirectLinkHelper->GetSourceName() : FString();
 }
 
 bool ADatasmithRuntimeActor::OpenConnectionWithIndex(int32 SourceIndex)
 {
 	using namespace DatasmithRuntime;
 
-	if (DirectLinkHelper->CanConnect())
+	if (DirectLinkHelper.IsValid() && DirectLinkHelper->CanConnect())
 	{
 		const TArray<FDatasmithRuntimeSourceInfo>& SourcesList = FDestinationProxy::GetListOfSources();
 
@@ -244,7 +244,7 @@ int32 ADatasmithRuntimeActor::GetSourceIndex()
 {
 	using namespace DatasmithRuntime;
 
-	if (DirectLinkHelper->IsConnected())
+	if (DirectLinkHelper.IsValid() && DirectLinkHelper->IsConnected())
 	{
 		const TArray<FDatasmithRuntimeSourceInfo>& SourcesList = FDestinationProxy::GetListOfSources();
 
@@ -262,7 +262,7 @@ int32 ADatasmithRuntimeActor::GetSourceIndex()
 
 void ADatasmithRuntimeActor::CloseConnection()
 {
-	if (DirectLinkHelper->IsConnected())
+	if (DirectLinkHelper.IsValid() && DirectLinkHelper->IsConnected())
 	{
 		DirectLinkHelper->CloseConnection();
 		Reset();
diff --git a/Engine/Plugins/Experimental/Enterprise/DatasmithRuntime/Source/Private/DatasmithRuntimeUtils.cpp b/Engine/Plugins/Experimental/Enterprise/DatasmithRuntime/Source/Private/DatasmithRuntimeUtils.cpp
--- a/Engine/Plugins/Experimental/Enterprise/DatasmithRuntime/Source/Private/DatasmithRuntimeUtils.cpp
+++ b/Engine/Plugins/Experimental/Enterprise/DatasmithRuntime/Source/Private/DatasmithRuntimeUtils.cpp
@@ -196,7 +196,7 @@ namespace DatasmithRuntime
 		}
 	}
 
-	TMap<DirectLink::FElementHash, TStrongObjectPtr<UObject>> FAssetRegistry::RegistrationMap;
+	TMap<uint32, TStrongObjectPtr<UObject>> FAssetRegistry::RegistrationMap;
 	TMap<uint32, TMap<FSceneGraphId,FAssetData>*> FAssetRegistry::SceneMappings;
 
 	union FRegistryKey
@@ -349,6 +349,19 @@ namespace DatasmithRuntime
 		return bIsCompleted;
 	}
 
+	int32 FAssetRegistry::GetAssetReferenceCount(UObject * Asset)
+	{
+		if (IInterface_AssetUserData* AssetUserData = Cast< IInterface_AssetUserData >(Asset))
+		{
+			if (UDatasmithRuntimeAuxiliaryData* AuxillaryData = AssetUserData->GetAssetUserData<UDatasmithRuntimeAuxiliaryData>())
+			{
+				return AuxillaryData->Referencers.Num();
+			}
+		}
+
+		return -1;
+	}
+
 	void FAssetRegistry::UnregisteredAssetsData(UObject* Asset, uint32 SceneKey, TFunction<void(FAssetData& AssetData)> UpdateFunc)
 	{
 		if (IInterface_AssetUserData* AssetUserData = Cast< IInterface_AssetUserData >(Asset))
@@ -394,7 +407,7 @@ namespace DatasmithRuntime
 		ensure(false);
 	}
 
-	UObject* FAssetRegistry::FindObjectFromHash(DirectLink::FElementHash ElementHash)
+	UObject* FAssetRegistry::FindObjectFromHash(uint32 ElementHash)
 	{
 		TStrongObjectPtr<UObject>* AssetPtr = RegistrationMap.Find(ElementHash);
 		return AssetPtr ? (*AssetPtr).Get() : nullptr;
@@ -402,10 +415,10 @@ namespace DatasmithRuntime
 
 	bool FAssetRegistry::CleanUp()
 	{
-		TArray<DirectLink::FElementHash> EntriesToDelete;
+		TArray<uint32> EntriesToDelete;
 		EntriesToDelete.Reserve(RegistrationMap.Num());
 
-		for (TPair<DirectLink::FElementHash, TStrongObjectPtr<UObject>>& Entry : RegistrationMap)
+		for (TPair<uint32, TStrongObjectPtr<UObject>>& Entry : RegistrationMap)
 		{
 			if (IInterface_AssetUserData* AssetUserData = Cast< IInterface_AssetUserData >(Entry.Value.Get()))
 			{
@@ -425,7 +438,7 @@ namespace DatasmithRuntime
 			}
 		}
 
-		for (DirectLink::FElementHash ElementHash : EntriesToDelete)
+		for (uint32 ElementHash : EntriesToDelete)
 		{
 			RegistrationMap.Remove(ElementHash);
 		}
diff --git a/Engine/Plugins/Experimental/Enterprise/DatasmithRuntime/Source/Private/DatasmithRuntimeUtils.h b/Engine/Plugins/Experimental/Enterprise/DatasmithRuntime/Source/Private/DatasmithRuntimeUtils.h
--- a/Engine/Plugins/Experimental/Enterprise/DatasmithRuntime/Source/Private/DatasmithRuntimeUtils.h
+++ b/Engine/Plugins/Experimental/Enterprise/DatasmithRuntime/Source/Private/DatasmithRuntimeUtils.h
@@ -84,13 +84,15 @@ namespace DatasmithRuntime
 
 		static bool IsObjectCompleted(UObject* Asset);
 
-		static UObject* FindObjectFromHash(DirectLink::FElementHash ElementHash);
+		static UObject* FindObjectFromHash(uint32 ElementHash);
+
+		static int32 GetAssetReferenceCount(UObject* Asset);
 
 		/** @return true if some assets have been marked for deletion */
 		static bool CleanUp();
 
 	private:
-		static TMap<DirectLink::FElementHash, TStrongObjectPtr<UObject>> RegistrationMap;
+		static TMap<uint32, TStrongObjectPtr<UObject>> RegistrationMap;
 		static TMap<uint32, TMap<FSceneGraphId,FAssetData>*> SceneMappings;
 	};
 
diff --git a/Engine/Plugins/Experimental/Enterprise/DatasmithRuntime/Source/Private/MeshImporter.cpp b/Engine/Plugins/Experimental/Enterprise/DatasmithRuntime/Source/Private/MeshImporter.cpp
--- a/Engine/Plugins/Experimental/Enterprise/DatasmithRuntime/Source/Private/MeshImporter.cpp
+++ b/Engine/Plugins/Experimental/Enterprise/DatasmithRuntime/Source/Private/MeshImporter.cpp
@@ -51,25 +51,22 @@ namespace DatasmithRuntime
 
 		TSharedPtr< IDatasmithMeshElement > MeshElement = StaticCastSharedPtr< IDatasmithMeshElement >(Elements[MeshData.ElementId]);
 
-		const int32 MaterialSlotCount = MeshElement->GetMaterialSlotCount();
-
 		UStaticMesh* StaticMesh = MeshData.GetObject<UStaticMesh>();
 
-		// Check if geometry has not changed
+		// If static mesh already completed, check if geometry has changed
 		if (StaticMesh)
 		{
+			uint32 NewResourceHash = GetTypeHash(MeshElement->GetFileHash());
 			// Force recreation of the static mesh if the mesh's file has changed
-			DirectLink::FElementHash NewResourceHash = GetTypeHash(MeshElement->GetFileHash());
 			if (MeshData.ResourceHash != NewResourceHash)
 			{
+				MeshData.ClearState(EAssetState::Completed);
 				FAssetRegistry::UnregisterAssetData(StaticMesh, SceneKey, MeshData.ElementId);
 				StaticMesh = nullptr;
 				MeshData.Object.Reset();
 			}
 		}
 
-		bool bUsingStaticMeshFromCache = false;
-
 		if (StaticMesh == nullptr)
 		{
 			MeshData.Hash = GetTypeHash(MeshElement->CalculateElementHash(true), EDataType::Mesh);
@@ -79,10 +76,6 @@ namespace DatasmithRuntime
 			{
 				StaticMesh = Cast<UStaticMesh>(AssetPtr);
 				check(StaticMesh);
-
-				MeshData.Object = TWeakObjectPtr<UObject>(StaticMesh);
-
-				bUsingStaticMeshFromCache = true;
 			}
 			else
 			{
@@ -96,12 +89,10 @@ namespace DatasmithRuntime
 #endif
 				check(StaticMesh);
 
-				MeshData.Object = TWeakObjectPtr< UObject >(StaticMesh);
-
 				// Add the creation of the mesh to the queue
 				FActionTaskFunction TaskFunc = [this](UObject* Object, const FReferencer& Referencer) -> EActionResult::Type
 				{
-					OnGoingTasks.Emplace( Async(
+					this->OnGoingTasks.Emplace( Async(
 #if WITH_EDITOR
 						EAsyncExecution::LargeThreadPool,
 #else
@@ -125,79 +116,40 @@ namespace DatasmithRuntime
 
 				MeshElementSet.Add(MeshData.ElementId);
 			}
+
+			MeshData.Object = TWeakObjectPtr< UObject >(StaticMesh);
 		}
 
-		// Collect all materials used by the static mesh.
-		if (!bUsingStaticMeshFromCache && MaterialSlotCount > 0)
+		// Collect materials used by static mesh
+		for (int32 Index = 0; Index < MeshElement->GetMaterialSlotCount(); Index++)
 		{
-			StaticMesh->GetStaticMaterials().SetNum(MaterialSlotCount);
-
-			for (int32 Index = 0; Index < MaterialSlotCount; Index++)
+			if (const IDatasmithMaterialIDElement* MaterialIDElement = MeshElement->GetMaterialSlotAt(Index).Get())
 			{
-				if (const IDatasmithMaterialIDElement* MaterialIDElement = MeshElement->GetMaterialSlotAt(Index).Get())
-				{
-					const FString MaterialPathName(MaterialIDElement->GetName());
+				const FString MaterialPathName(MaterialIDElement->GetName());
 
-					if (!MaterialPathName.StartsWith(TEXT("/")))
+				if (!MaterialPathName.StartsWith(TEXT("/")))
+				{
+					if (FSceneGraphId* MaterialElementIdPtr = AssetElementMapping.Find(MaterialPrefix + MaterialPathName))
 					{
-						if (FSceneGraphId* MaterialElementIdPtr = AssetElementMapping.Find(MaterialPrefix + MaterialPathName))
-						{
-							FAssetData& MaterialData = AssetDataList[*MaterialElementIdPtr];
-
-							ProcessMaterialData(MaterialData);
-						}
+						ProcessMaterialData(AssetDataList[*MaterialElementIdPtr]);
 					}
-					else
+				}
+				else
+				{
+					// Force loading of material asset if it exists. It will be assigned when the mesh is built
+					UMaterialInterface* MaterialInterface = Cast<UMaterialInterface>(FSoftObjectPath(MaterialPathName).TryLoad());
+					if (MaterialInterface == nullptr)
 					{
-						// Force loading of material asset if it exists. It will be assigned when the mesh is built
-						UMaterialInterface* MaterialInterface = Cast<UMaterialInterface>(FSoftObjectPath(MaterialPathName).TryLoad());
 						UE_LOG(LogDatasmithRuntime, Warning, TEXT("ProcessMeshData: Cannot find material %s"), *MaterialPathName);
 					}
 				}
 			}
 		}
-		else
-		{
-			ensure(StaticMesh->GetStaticMaterials().Num() == MaterialSlotCount);
-		}
 
-		if (!bUsingStaticMeshFromCache)
+		// If static mesh already completed, then material assignment has changed
+		if (MeshData.HasState(EAssetState::Completed))
 		{
-			TArray< FStaticMaterial >& StaticMaterials = StaticMesh->GetStaticMaterials();
-			StaticMaterials.SetNum(MaterialSlotCount);
-
-			for (int32 Index = 0; Index < MaterialSlotCount; Index++)
-			{
-				FStaticMaterial& StaticMaterial = StaticMaterials[Index];
-
-				StaticMaterial.MaterialSlotName = NAME_None;
-				StaticMaterial.MaterialInterface = nullptr;
-
-				// Done to remove an assert from an 'ensure' in UStaticMesh::GetUVChannelData
-				StaticMaterial.UVChannelData = FMeshUVChannelInfo(1.f);
-
-				if (const IDatasmithMaterialIDElement* MaterialIDElement = MeshElement->GetMaterialSlotAt(Index).Get())
-				{
-					const FString MaterialPathName(MaterialIDElement->GetName());
-
-					if (!MaterialPathName.StartsWith(TEXT("/")))
-					{
-						if (FSceneGraphId* MaterialElementIdPtr = AssetElementMapping.Find(MaterialPrefix + MaterialPathName))
-						{
-							ProcessMaterialData(AssetDataList[*MaterialElementIdPtr]);
-						}
-					}
-					else
-					{
-						StaticMaterial.MaterialInterface = Cast<UMaterialInterface>(FSoftObjectPath(MaterialPathName).TryLoad());
-					}
-
-					if (!bUsingStaticMeshFromCache)
-					{
-						StaticMaterial.MaterialSlotName = *FString::Printf(TEXT("%d"), MaterialIDElement->GetId());
-					}
-				}
-			}
+			UpdateStaticMeshMaterials(MeshData);
 		}
 
 		// Create BodySetup in game thread to avoid allocating during a garbage collect later on
@@ -544,18 +496,9 @@ namespace DatasmithRuntime
 		{
 			FGCScopeGuard GCGuard;
 
-#if !WITH_EDITOR
-			// Force build process to keep index buffer for complex collision
-			const bool bKeepIndexBuffer = ImportOptions.BuildCollisions != ECollisionEnabled::NoCollision && (ImportOptions.CollisionType == ECollisionTraceFlag::CTF_UseComplexAsSimple || ImportOptions.CollisionType == ECollisionTraceFlag::CTF_UseSimpleAndComplex);
-			if (bKeepIndexBuffer)
-			{
-				StaticMesh->bAllowCPUAccess = true;
-			}
-#endif
-
-			// Do not mark the package dirty since MarkPackageDirty is not thread safe
 			UStaticMesh::FBuildMeshDescriptionsParams Params;
 			Params.bUseHashAsGuid = true;
+			// Do not mark the package dirty since MarkPackageDirty is not thread safe
 			Params.bMarkPackageDirty = false;
 			Params.bBuildSimpleCollision = false;
 			// Do not commit since we only need the render data and commit is slow
@@ -884,10 +827,10 @@ namespace DatasmithRuntime
 
 	void FSceneImporter::UpdateStaticMeshMaterials(FAssetData& MeshData)
 	{
-		TSharedRef< IDatasmithMeshElement > MeshElement = StaticCastSharedPtr< IDatasmithMeshElement >(Elements[MeshData.ElementId]).ToSharedRef();
+		TSharedPtr< IDatasmithMeshElement > MeshElement = StaticCastSharedPtr< IDatasmithMeshElement >(Elements[MeshData.ElementId]);
 		UStaticMesh* StaticMesh = MeshData.GetObject<UStaticMesh>();
 
-		if (StaticMesh == nullptr)
+		if (!MeshElement.IsValid() || StaticMesh == nullptr)
 		{
 			return;
 		}
@@ -931,11 +874,14 @@ namespace DatasmithRuntime
 				// Mark dependent mesh components' render state as dirty
 				for (FReferencer& ActorReferencer : MeshData.Referencers)
 				{
-					FActorData& ActorData = ActorDataList[ActorReferencer.GetId()];
+					const FActorData& ActorData = ActorDataList[ActorReferencer.GetId()];
 
-					if (UActorComponent* ActorComponent = ActorData.GetObject<UActorComponent>())
+					if (ActorData.HasState(EAssetState::Completed))
 					{
-						ActorComponent->MarkRenderStateDirty();
+						if (UActorComponent* ActorComponent = ActorData.GetObject<UActorComponent>())
+						{
+							ActorComponent->MarkRenderStateDirty();
+						}
 					}
 				}
 			}
diff --git a/Engine/Plugins/Experimental/Enterprise/DatasmithRuntime/Source/Private/SceneImporter.cpp b/Engine/Plugins/Experimental/Enterprise/DatasmithRuntime/Source/Private/SceneImporter.cpp
--- a/Engine/Plugins/Experimental/Enterprise/DatasmithRuntime/Source/Private/SceneImporter.cpp
+++ b/Engine/Plugins/Experimental/Enterprise/DatasmithRuntime/Source/Private/SceneImporter.cpp
@@ -479,6 +479,9 @@ namespace DatasmithRuntime
 		// Flag used to avoid deleting components and associated assets in the same frame
 		if (bContinue && EnumHasAnyFlags(TasksToComplete, EWorkerTask::DeleteComponent))
 		{
+			// Terminate all rendering commands before deleting any component
+			FlushRenderingCommands();
+
 			FActionTask ActionTask;
 			while (FPlatformTime::Seconds() < EndTime)
 			{
@@ -495,8 +498,12 @@ namespace DatasmithRuntime
 		// Force a garbage collection if we are done with the components
 		if (ActionQueues[EQueueTask::DeleteCompQueue].IsEmpty() && EnumHasAnyFlags(TasksToComplete, EWorkerTask::GarbageCollect))
 		{
+			// Terminate all rendering commands before deleting any component
+			FlushRenderingCommands();
+
 			if (!IsGarbageCollecting())
 			{
+				CollectGarbage(GARBAGE_COLLECTION_KEEPFLAGS);
 				TasksToComplete &= ~EWorkerTask::GarbageCollect;
 			}
 		}
@@ -521,6 +528,9 @@ namespace DatasmithRuntime
 
 		if (TasksToComplete == EWorkerTask::NoTask && SceneElement.IsValid())
 		{
+			// Terminate all rendering commands before deleting any asset
+			FlushRenderingCommands();
+
 			// Delete assets which has not been reused on the last processing
 			if (FAssetRegistry::CleanUp())
 			{
@@ -617,7 +627,7 @@ namespace DatasmithRuntime
 
 		PrepareIncrementalUpdate(UpdateContext);
 
-		IncrementalAdditions(UpdateContext.Additions);
+		IncrementalAdditions(UpdateContext.Additions, UpdateContext.Updates);
 
 		IncrementalModifications(UpdateContext.Updates);
 
@@ -649,11 +659,18 @@ namespace DatasmithRuntime
 							continue;
 						}
 
+						const EDataType DataType = AssetDataList[ElementId].Type;
+						const FString& Prefix = DataType == EDataType::Texture ? TexturePrefix : (DataType == EDataType::Material ? MaterialPrefix : MeshPrefix);
+
+						UE_LOG(LogDatasmithRuntime, Log, TEXT("IncrementalDeletions: %s %d"), *Prefix, ElementId);
+
 						AddToQueue(EQueueTask::DeleteAssetQueue, { TaskFunc, FReferencer(ElementId) } );
 						TasksToComplete |= EWorkerTask::DeleteAsset;
 					}
 					else if (ActorDataList.Contains(ElementId))
 					{
+						UE_LOG(LogDatasmithRuntime, Log, TEXT("IncrementalDeletions: actor %s (%d)"), Elements[ElementId]->GetLabel(), ElementId);
+
 						HideSceneComponent(ActorDataList[ElementId].GetObject<USceneComponent>());
 
 						AddToQueue(EQueueTask::DeleteCompQueue, { TaskFunc, FReferencer(ElementId) } );
@@ -661,8 +678,7 @@ namespace DatasmithRuntime
 					}
 					else
 					{
-						TSharedPtr<IDatasmithElement> Element = Elements[ElementId];
-						UE_LOG(LogDatasmithRuntime, Error, TEXT("Element %d (%s) was not found"), ElementId, Element->GetName());
+						UE_LOG(LogDatasmithRuntime, Error, TEXT("Element %d (%s) was not found"), ElementId, Elements[ElementId]->GetName());
 						ensure(false);
 					}
 				}
@@ -671,11 +687,6 @@ namespace DatasmithRuntime
 
 		bIncrementalUpdate = true;
 
-		for (TPair< FSceneGraphId, FAssetData >& Entry : AssetDataList)
-		{
-			Entry.Value.Referencers.Empty();
-		}
-
 		// Parse scene to collect all actions to be taken
 		for (int32 Index = 0; Index < SceneElement->GetActorsCount(); ++Index)
 		{
@@ -714,6 +725,8 @@ namespace DatasmithRuntime
 					const EDataType DataType = AssetDataList[ElementId].Type;
 					const FString& Prefix = DataType == EDataType::Texture ? TexturePrefix : (DataType == EDataType::Material ? MaterialPrefix : MeshPrefix);
 
+					UE_LOG(LogDatasmithRuntime, Log, TEXT("IncrementalModifications: %s %d"), *Prefix, ElementPtr->GetNodeId());
+
 					const FString PrefixedName = Prefix + ElementPtr->GetName();
 
 					if (!AssetElementMapping.Contains(PrefixedName))
@@ -823,7 +836,7 @@ namespace DatasmithRuntime
 		}
 	}
 
-	void FSceneImporter::IncrementalAdditions(TArray<TSharedPtr<IDatasmithElement>>& Additions)
+	void FSceneImporter::IncrementalAdditions(TArray<TSharedPtr<IDatasmithElement>>& Additions, TArray<TSharedPtr<IDatasmithElement>>& Updates)
 	{
 		if (Additions.Num() == 0)
 		{
@@ -841,26 +854,8 @@ namespace DatasmithRuntime
 		Elements.Reserve( Elements.Num() + AdditionCount );
 		AssetDataList.Reserve( AssetDataList.Num() + AdditionCount );
 
-		TFunction<void(FAssetData&)> UpdateReference;
 		TFunction<void(TSharedPtr<IDatasmithElement>&&, EDataType)> LocalAddAsset;
 
-		UpdateReference = [&](FAssetData& AssetData) -> void
-		{
-			AssetData.ClearState(EAssetState::Processed);
-
-			for (const FReferencer& Referencer : AssetData.Referencers)
-			{
-				if (this->AssetDataList.Contains(Referencer.ElementId))
-				{
-					UpdateReference(this->AssetDataList[Referencer.ElementId]);
-				}
-				else if (this->ActorDataList.Contains(Referencer.ElementId))
-				{
-					this->ActorDataList[Referencer.ElementId].ClearState(EAssetState::Processed);
-				}
-			}
-		};
-
 		LocalAddAsset = [&](TSharedPtr<IDatasmithElement>&& Element, EDataType DataType) -> void
 		{
 			const FString& Prefix = DataType == EDataType::Texture ? TexturePrefix : (DataType == EDataType::Material ? MaterialPrefix : MeshPrefix);
@@ -872,16 +867,26 @@ namespace DatasmithRuntime
 			if (this->AssetElementMapping.Contains(PrefixedName))
 			{
 				const FSceneGraphId ExistingElementId = AssetElementMapping[PrefixedName];
-				UE_LOG(LogDatasmithRuntime, Warning, TEXT("Found a new Element (%d) with the same name, %s, as an existing one (%d). Replacing the existing one ..."), ElementId, this->Elements[ExistingElementId]->GetName(), ExistingElementId);
+				FAssetData& ExistingAssetData = AssetDataList[ExistingElementId];
 
-				const FAssetData& AssetData = this->AssetDataList[ExistingElementId];
-				ensure(AssetData.HasState(EAssetState::PendingDelete));
+				if (!ExistingAssetData.HasState(EAssetState::PendingDelete))
+				{
+					UE_LOG(LogDatasmithRuntime, Error, TEXT("Found a new %s (%d) with the same name, %s, as an existing one (%d)."), *Prefix, ElementId, this->Elements[ExistingElementId]->GetName(), ExistingElementId);
+				}
 
-				for (const FReferencer& Referencer : AssetData.Referencers)
+				// Add all referencers to the list of elements to update
+				for (const FReferencer& Referencer : ExistingAssetData.Referencers)
 				{
 					if (this->AssetDataList.Contains(Referencer.ElementId))
 					{
-						UpdateReference(this->AssetDataList[Referencer.ElementId]);
+						FAssetData& ReferencerAssetData = this->AssetDataList[Referencer.ElementId];
+						const bool bMustBeProcessed = ReferencerAssetData.HasState(EAssetState::Processed | EAssetState::Completed) && !ReferencerAssetData.HasState(EAssetState::PendingDelete);
+						
+						if (bMustBeProcessed)
+						{
+							ReferencerAssetData.ClearState(EAssetState::Processed);
+							Updates.Add(Elements[Referencer.ElementId]);
+						}
 					}
 					else if (this->ActorDataList.Contains(Referencer.ElementId))
 					{
@@ -907,10 +912,12 @@ namespace DatasmithRuntime
 		{
 			if (ElementPtr->IsA(EDatasmithElementType::BaseMaterial))
 			{
+				UE_LOG(LogDatasmithRuntime, Log, TEXT("IncrementalAdditions: Material %d"), ElementPtr->GetNodeId());
 				LocalAddAsset(MoveTemp(ElementPtr), EDataType::Material);
 			}
 			else if (ElementPtr->IsA(EDatasmithElementType::StaticMesh))
 			{
+				UE_LOG(LogDatasmithRuntime, Log, TEXT("IncrementalAdditions: StaticMesh %d"), ElementPtr->GetNodeId());
 				if (IDatasmithMeshElement* MeshElement = static_cast<IDatasmithMeshElement*>(ElementPtr.Get()))
 				{
 					// If resource file does not exist, add scene's resource path if valid
@@ -928,6 +935,7 @@ namespace DatasmithRuntime
 			}
 			else if (ElementPtr->IsA(EDatasmithElementType::Texture))
 			{
+				UE_LOG(LogDatasmithRuntime, Log, TEXT("IncrementalAdditions: Texture %d"), ElementPtr->GetNodeId());
 				if (IDatasmithTextureElement* TextureElement = static_cast<IDatasmithTextureElement*>(ElementPtr.Get()))
 				{
 					// If resource file does not exist, add scene's resource path if valid
@@ -1288,7 +1296,7 @@ namespace DatasmithRuntime
 
 			if (CameraComponent == nullptr)
 			{
-				if (ImportOptions.BuildHierarchy == EBuildHierarchyMethod::Unfiltered)
+				if (ImportOptions.BuildHierarchy != EBuildHierarchyMethod::None)
 				{
 					UWorld* World = RootComponent->GetOwner()->GetWorld();
 
diff --git a/Engine/Plugins/Experimental/Enterprise/DatasmithRuntime/Source/Private/SceneImporter.h b/Engine/Plugins/Experimental/Enterprise/DatasmithRuntime/Source/Private/SceneImporter.h
--- a/Engine/Plugins/Experimental/Enterprise/DatasmithRuntime/Source/Private/SceneImporter.h
+++ b/Engine/Plugins/Experimental/Enterprise/DatasmithRuntime/Source/Private/SceneImporter.h
@@ -234,7 +234,7 @@ namespace DatasmithRuntime
 		}
 
 		template<typename T = UObject>
-		T* GetObject()
+		T* GetObject() const
 		{
 			return Cast< T >(Object.Get());
 		}
@@ -249,10 +249,10 @@ namespace DatasmithRuntime
 		int32 Requirements = 0;
 
 		/** Hash of associated element used to prevent the duplication of assets */
-		DirectLink::FElementHash Hash = 0;
+		uint32 Hash = 0;
 
 		/** Hash of potential resource of associated element used to prevent recreation of assets */
-		DirectLink::FElementHash ResourceHash = 0;
+		uint32 ResourceHash = 0;
 
 		FAssetData(FSceneGraphId InElementId, EDataType InType = EDataType::None)
 			: FBaseData(InElementId, InType)
@@ -524,7 +524,7 @@ namespace DatasmithRuntime
 
 		void PrepareIncrementalUpdate(FUpdateContext& UpdateContext);
 
-		void IncrementalAdditions(TArray<TSharedPtr<IDatasmithElement>>& Additions);
+		void IncrementalAdditions(TArray<TSharedPtr<IDatasmithElement>>& Additions, TArray<TSharedPtr<IDatasmithElement>>& Updates);
 
 		void IncrementalModifications(TArray<TSharedPtr<IDatasmithElement>>& Modifications);
 
diff --git a/Engine/Plugins/Experimental/Enterprise/DatasmithRuntime/Source/Private/TextureImporter.cpp b/Engine/Plugins/Experimental/Enterprise/DatasmithRuntime/Source/Private/TextureImporter.cpp
--- a/Engine/Plugins/Experimental/Enterprise/DatasmithRuntime/Source/Private/TextureImporter.cpp
+++ b/Engine/Plugins/Experimental/Enterprise/DatasmithRuntime/Source/Private/TextureImporter.cpp
@@ -227,7 +227,7 @@ namespace DatasmithRuntime
 
 		if (Texture)
 		{
-			DirectLink::FElementHash TextureHash = GetTypeHash(TextureElement->CalculateElementHash(true), EDataType::Texture);
+			uint32 TextureHash = GetTypeHash(TextureElement->CalculateElementHash(true), EDataType::Texture);
 
 			FAssetRegistry::UnregisteredAssetsData(THelper, SceneKey, [this, &Texture, TextureHash](FAssetData& AssetData) -> void
 				{
@@ -338,7 +338,7 @@ namespace DatasmithRuntime
 
 		IDatasmithTextureElement* TextureElement = static_cast<IDatasmithTextureElement*>(Elements[ TextureId ].Get());
 
-		DirectLink::FElementHash TextureHash = GetTypeHash(TextureElement->CalculateElementHash(true), EDataType::Texture);
+		uint32 TextureHash = GetTypeHash(TextureElement->CalculateElementHash(true), EDataType::Texture);
 
 		if (UObject* Asset = FAssetRegistry::FindObjectFromHash(TextureHash))
 		{
