Fixed build break when compiling QAEnterprise game (reported by horde)

diff --git a/Engine/Plugins/Experimental/Enterprise/DatasmithRuntime/Source/Private/MeshImporter.cpp b/Engine/Plugins/Experimental/Enterprise/DatasmithRuntime/Source/Private/MeshImporter.cpp
--- a/Engine/Plugins/Experimental/Enterprise/DatasmithRuntime/Source/Private/MeshImporter.cpp
+++ b/Engine/Plugins/Experimental/Enterprise/DatasmithRuntime/Source/Private/MeshImporter.cpp
@@ -504,6 +504,11 @@ namespace DatasmithRuntime
 			// Do not commit since we only need the render data and commit is slow
 			Params.bCommitMeshDescription = false;
 			Params.bFastBuild = true;
+#if !WITH_EDITOR
+			// Force build process to keep index buffer for complex collision when in game
+			//Params.bAllowCpuAccess = ImportOptions.BuildCollisions != ECollisionEnabled::NoCollision && (ImportOptions.CollisionType == ECollisionTraceFlag::CTF_UseComplexAsSimple || ImportOptions.CollisionType == ECollisionTraceFlag::CTF_UseSimpleAndComplex);
+#endif
+
 			StaticMesh->BuildFromMeshDescriptions(MeshDescriptionPointers, Params);
 		}
 
