Added code to set the name of area lights

diff --git a/Engine/Plugins/Experimental/Enterprise/DatasmithRuntime/Source/Private/LightImporter.cpp b/Engine/Plugins/Experimental/Enterprise/DatasmithRuntime/Source/Private/LightImporter.cpp
--- a/Engine/Plugins/Experimental/Enterprise/DatasmithRuntime/Source/Private/LightImporter.cpp
+++ b/Engine/Plugins/Experimental/Enterprise/DatasmithRuntime/Source/Private/LightImporter.cpp
@@ -349,6 +349,7 @@ namespace DatasmithRuntime
 
 			if ( LightShapeActor )
 			{
+				LightShapeActor->Rename(AreaLightElement->GetName(), nullptr, REN_NonTransactional | REN_DontCreateRedirectors);
 #if WITH_EDITOR
 				LightShapeActor->SetActorLabel( AreaLightElement->GetLabel() );
 #endif
