Fixed issue with collisions not working in packaged game

diff --git a/Engine/Plugins/Experimental/Enterprise/DatasmithRuntime/Source/Private/DatasmithRuntimeUtils.cpp b/Engine/Plugins/Experimental/Enterprise/DatasmithRuntime/Source/Private/DatasmithRuntimeUtils.cpp
--- a/Engine/Plugins/Experimental/Enterprise/DatasmithRuntime/Source/Private/DatasmithRuntimeUtils.cpp
+++ b/Engine/Plugins/Experimental/Enterprise/DatasmithRuntime/Source/Private/DatasmithRuntimeUtils.cpp
@@ -12,7 +12,6 @@
 #include "Algo/AnyOf.h"
 #include "Async/Async.h"
 #include "Engine/Polys.h"
-#include "Engine/StaticMesh.h"
 #include "HAL/FileManager.h"
 #include "Math/Plane.h"
 #include "MeshDescription.h"
@@ -22,6 +21,8 @@
 #include "PhysicsEngine/AggregateGeom.h"
 #include "PhysicsEngine/BodySetup.h"
 #include "PhysicsEngine/PhysicsSettings.h"
+#include "PhysicsPublicCore.h"
+#include "PhysXCookHelper.h"
 #include "StaticMeshAttributes.h"
 #include "StaticMeshOperations.h"
 
@@ -446,7 +447,8 @@ namespace DatasmithRuntime
 		return EntriesToDelete.Num() > 0;
 	}
 
-	// k-DOP (k-Discrete Oriented Polytopes) Direction Vectors
+	// Code below has been borrowed from GenerateKDopAsSimpleCollision in GeomFitUtils.cpp
+	// to generate k-DOP (k-Discrete Oriented Polytopes) with 26 polytopes
 	#define RCP_SQRT2 (0.70710678118654752440084436210485f)
 	#define RCP_SQRT3 (0.57735026918962576450914878050196f)
 
@@ -479,10 +481,8 @@ namespace DatasmithRuntime
 		FVector(-RCP_SQRT3, -RCP_SQRT3,  RCP_SQRT3),
 		FVector(-RCP_SQRT3, -RCP_SQRT3, -RCP_SQRT3),
 	};
-	// THIS FUNCTION REPLACES EXISTING SIMPLE COLLISION MODEL WITH KDOP
-#define MY_FLTMAX (3.402823466e+38F)
-	constexpr float HalfWorldMax = HALF_WORLD_MAX;
 
+	constexpr float HalfWorldMax = HALF_WORLD_MAX;
 
 	void GenerateKDopAsSimpleCollision(UBodySetup* BodySetup, const FStaticMeshLODResources& Resources, const FVector* Directions, int32 DirectionCount)
 	{
@@ -491,7 +491,7 @@ namespace DatasmithRuntime
 
 		for (int32 Index = 0; Index < DirectionCount; ++Index)
 		{
-			MaxDistances.Add(-MY_FLTMAX);
+			MaxDistances.Add(-MAX_FLT);
 		}
 
 		// For each vertex, project along each kdop direction, to find the max in that direction.
@@ -581,6 +581,7 @@ namespace DatasmithRuntime
 
 		ConvexElem.UpdateElemBox();
 	}
+	// End of Code borrowed from GenerateKDopAsSimpleCollision in GeomFitUtils.cpp
 
 	void BuildCollision(UBodySetup* BodySetup, ECollisionTraceFlag CollisionFlag, const FStaticMeshLODResources& Resources)
 	{
@@ -591,48 +592,16 @@ namespace DatasmithRuntime
 			BodySetup->CollisionTraceFlag = UPhysicsSettings::Get()->DefaultShapeComplexity;
 		}
 
-		BodySetup->AggGeom.EmptyElements();
-
-		// Use the bounding box as the collision mesh if simple collision is required and no aggregate geometry is provided
-		// #ue_dsruntime - TODO Choose the best shape according to dimension of bounding box
-		if (BodySetup->CollisionTraceFlag != ECollisionTraceFlag::CTF_UseComplexAsSimple)
+		// Use k-DOP 26 as the collision mesh if simple collision is required
+		// #ue_dsruntime - TODO | Choose better shape of collision mesh
+		if (BodySetup->CollisionTraceFlag != ECollisionTraceFlag::CTF_UseComplexAsSimple && BodySetup->AggGeom.ConvexElems.Num() == 0)
 		{
 			GenerateKDopAsSimpleCollision(BodySetup, Resources, KDopDir26, 26);
 		}
 
-#if WITH_EDITOR
-		BodySetup->CreatePhysicsMeshes();
-#else
-		// Create a convex hull element from the mesh description of the static mesh if complex as simple is required
-		if (BodySetup->CollisionTraceFlag != ECollisionTraceFlag::CTF_UseSimpleAsComplex)
-		{
-			FKConvexElem& ConvexElem = BodySetup->AggGeom.ConvexElems.AddDefaulted_GetRef();
-
-			const FPositionVertexBuffer& PositionVertexBuffer = Resources.VertexBuffers.PositionVertexBuffer;
-			ConvexElem.VertexData.Reserve(Resources.GetNumVertices());
-
-			for(int32 Index = 0; Index < Resources.GetNumVertices(); ++Index)
-			{
-				ConvexElem.VertexData.Add(PositionVertexBuffer.VertexPosition(Index));
-			}
-
-			ConvexElem.UpdateElemBox();
-		}
-
-		BodySetup->CollisionTraceFlag = ECollisionTraceFlag::CTF_UseSimpleAsComplex;
-
+		// Creation of collision meshes can only happen on game thread
 		AsyncTask(ENamedThreads::GameThread, [BodySetup] {
-
-			BodySetup->InvalidatePhysicsData();
-
-			BodySetup->CreatePhysicsMeshesAsync(FOnAsyncPhysicsCookFinished::CreateLambda([](bool bSuccess, UBodySetup* BodySetup) -> void
-				{
-					FString Message(bSuccess ? TEXT("SUCCESS") : TEXT("FAILED"));
-					//UE_LOG(LogDatasmithRuntime, Log, TEXT("CreateStaticMesh: CreatePhysicsMeshesAsync %s."), *Message);
-				}
-				, BodySetup)
-			);
+			BodySetup->CreatePhysicsMeshes();
 		});
-#endif
 	}
 } // End of namespace DatasmithRuntime
\ No newline at end of file
diff --git a/Engine/Plugins/Experimental/Enterprise/DatasmithRuntime/Source/Private/DatasmithRuntimeUtils.h b/Engine/Plugins/Experimental/Enterprise/DatasmithRuntime/Source/Private/DatasmithRuntimeUtils.h
--- a/Engine/Plugins/Experimental/Enterprise/DatasmithRuntime/Source/Private/DatasmithRuntimeUtils.h
+++ b/Engine/Plugins/Experimental/Enterprise/DatasmithRuntime/Source/Private/DatasmithRuntimeUtils.h
@@ -8,8 +8,11 @@
 
 #include "DirectLinkCommon.h"
 
+#include "Engine/StaticMesh.h"
 #include "Misc/SecureHash.h"
 
+#include "DatasmithRuntimeUtils.generated.h"
+
 class IDatasmithBaseMaterialElement;
 class IDatasmithElement;
 class IDatasmithMasterMaterialElement;
@@ -22,7 +25,7 @@ class UClass;
 class UMaterial;
 class UMaterialInstanceDynamic;
 class USceneComponent;
-class UStaticMesh;
+class UWorld;
 
 struct FTextureData;
 struct FActorData;
@@ -30,6 +33,33 @@ struct FDatasmithMeshElementPayload;
 struct FMeshDescription;
 struct FStaticMeshLODResources;
 
+// Class deriving from UStaticMesh to allow the cooking of collision meshes at runtime
+// To do so, bAllowCPUAccess must be true AND  the metod GetWorld() must return a valid world
+UCLASS()
+class URuntimeMesh : public UStaticMesh
+{
+	GENERATED_BODY()
+
+public:
+	URuntimeMesh()
+		: World(nullptr)
+	{
+		// Set bAllowCPUAccess to true to allow the copy of render data triangles to the collision mesh
+		bAllowCPUAccess = true;
+	}
+
+	// UObject overrides
+	// Overridden to allow cooking of collision meshes, simple and complex, from static mesh at runtime
+	virtual UWorld* GetWorld() const override  { return World ? World : UStaticMesh::GetWorld(); }
+	// End UObject overrides
+
+	// Use valid world to allow cooking of collision meshes, simple and complex, from static mesh at runtime 
+	void SetWorld(UWorld* InWorld) { World = InWorld; }
+
+private:
+	UWorld* World;
+};
+
 namespace DatasmithRuntime
 {
 	enum class EDSResizeTextureMode
diff --git a/Engine/Plugins/Experimental/Enterprise/DatasmithRuntime/Source/Private/MeshImporter.cpp b/Engine/Plugins/Experimental/Enterprise/DatasmithRuntime/Source/Private/MeshImporter.cpp
--- a/Engine/Plugins/Experimental/Enterprise/DatasmithRuntime/Source/Private/MeshImporter.cpp
+++ b/Engine/Plugins/Experimental/Enterprise/DatasmithRuntime/Source/Private/MeshImporter.cpp
@@ -51,7 +51,7 @@ namespace DatasmithRuntime
 
 		TSharedPtr< IDatasmithMeshElement > MeshElement = StaticCastSharedPtr< IDatasmithMeshElement >(Elements[MeshData.ElementId]);
 
-		UStaticMesh* StaticMesh = MeshData.GetObject<UStaticMesh>();
+		URuntimeMesh* StaticMesh = MeshData.GetObject<URuntimeMesh>();
 
 		// If static mesh already completed, check if geometry has changed
 		if (StaticMesh)
@@ -74,7 +74,7 @@ namespace DatasmithRuntime
 
 			if (UObject* AssetPtr = FAssetRegistry::FindObjectFromHash(MeshData.Hash))
 			{
-				StaticMesh = Cast<UStaticMesh>(AssetPtr);
+				StaticMesh = Cast<URuntimeMesh>(AssetPtr);
 				check(StaticMesh);
 			}
 			else
@@ -83,12 +83,14 @@ namespace DatasmithRuntime
 				FString MeshName = TEXT("S_") + FString(MeshElement->GetName()) + TEXT("_LU_") + FString::FromInt(MeshData.ElementId);
 				MeshName = FDatasmithUtils::SanitizeObjectName(MeshName);
 				UPackage* Package = CreatePackage(*FPaths::Combine( TEXT("/DatasmithContent/Meshes"), MeshName));
-				StaticMesh = NewObject< UStaticMesh >(Package, *MeshName, RF_Public);
+				StaticMesh = NewObject< URuntimeMesh >(Package, *MeshName, RF_Public);
 #else
-				StaticMesh = NewObject< UStaticMesh >(GetTransientPackage());
+				StaticMesh = NewObject< URuntimeMesh >(GetTransientPackage());
 #endif
 				check(StaticMesh);
 
+				StaticMesh->SetWorld(RootComponent->GetWorld());
+
 				// Add the creation of the mesh to the queue
 				FActionTaskFunction TaskFunc = [this](UObject* Object, const FReferencer& Referencer) -> EActionResult::Type
 				{
@@ -367,7 +369,11 @@ namespace DatasmithRuntime
 		int32 MaxLightmapSize = FDatasmithStaticMeshImportOptions::ConvertLightmapEnumToValue(EDatasmithImportLightmapMax::LIGHTMAP_512);
 
 		// 4. Collisions
-		ProcessCollision(StaticMesh, MeshPayload);
+		StaticMesh->GetBodySetup()->AggGeom.EmptyElements();
+		if (ImportOptions.BuildCollisions != ECollisionEnabled::NoCollision && ImportOptions.CollisionType != ECollisionTraceFlag::CTF_UseComplexAsSimple)
+		{
+			ProcessCollision(StaticMesh, MeshPayload);
+		}
 
 		// Extracted from FDatasmithStaticMeshImporter::SetupStaticMesh
 #if WITH_EDITOR
@@ -506,7 +512,7 @@ namespace DatasmithRuntime
 			Params.bFastBuild = true;
 #if !WITH_EDITOR
 			// Force build process to keep index buffer for complex collision when in game
-			//Params.bAllowCpuAccess = ImportOptions.BuildCollisions != ECollisionEnabled::NoCollision && (ImportOptions.CollisionType == ECollisionTraceFlag::CTF_UseComplexAsSimple || ImportOptions.CollisionType == ECollisionTraceFlag::CTF_UseSimpleAndComplex);
+			Params.bAllowCpuAccess = ImportOptions.BuildCollisions != ECollisionEnabled::NoCollision && (ImportOptions.CollisionType == ECollisionTraceFlag::CTF_UseComplexAsSimple || ImportOptions.CollisionType == ECollisionTraceFlag::CTF_UseSimpleAndComplex);
 #endif
 
 			StaticMesh->BuildFromMeshDescriptions(MeshDescriptionPointers, Params);
