Fixed flaw in logic in processing of deleted elements: The removal of them from the list of referencers was missing.
Completed the work related to the runtime material selector.

diff --git a/Engine/Plugins/Experimental/Enterprise/DatasmithRuntime/Source/Private/MaterialImportUtils.cpp b/Engine/Plugins/Experimental/Enterprise/DatasmithRuntime/Source/Private/MaterialImportUtils.cpp
--- a/Engine/Plugins/Experimental/Enterprise/DatasmithRuntime/Source/Private/MaterialImportUtils.cpp
+++ b/Engine/Plugins/Experimental/Enterprise/DatasmithRuntime/Source/Private/MaterialImportUtils.cpp
@@ -125,7 +125,7 @@ namespace DatasmithRuntime
 		return ParametersRef;
 	}
 
-	int32 ProcessMaterialElement(TSharedPtr< IDatasmithMasterMaterialElement > MasterMaterialElement, const TCHAR* Host, FTextureCallback TextureCallback)
+	int32 ProcessMaterialElement(TSharedPtr< IDatasmithMasterMaterialElement > MasterMaterialElement, FTextureCallback TextureCallback)
 	{
 		TRACE_CPUPROFILER_EVENT_SCOPE(DatasmithRuntime::ProcessMasterMaterialElement);
 
@@ -137,7 +137,7 @@ namespace DatasmithRuntime
 			return MaterialRequirement;
 		}
 
-		TSharedPtr< FDatasmithMasterMaterialSelector > MaterialSelector = FDatasmithMasterMaterialManager::Get().GetSelector(Host);
+		TSharedPtr< FDatasmithMasterMaterialSelector > MaterialSelector = FDatasmithMasterMaterialManager::Get().GetSelector(MATERIAL_HOST);
 
 		UMaterialInterface* Material = nullptr;
 
diff --git a/Engine/Plugins/Experimental/Enterprise/DatasmithRuntime/Source/Private/MaterialImportUtils.h b/Engine/Plugins/Experimental/Enterprise/DatasmithRuntime/Source/Private/MaterialImportUtils.h
--- a/Engine/Plugins/Experimental/Enterprise/DatasmithRuntime/Source/Private/MaterialImportUtils.h
+++ b/Engine/Plugins/Experimental/Enterprise/DatasmithRuntime/Source/Private/MaterialImportUtils.h
@@ -31,7 +31,7 @@ namespace DatasmithRuntime
 
 	TSharedPtr<IDatasmithUEPbrMaterialElement> ValidatePbrMaterial(TSharedPtr<IDatasmithUEPbrMaterialElement> PbrMaterialElement, FSceneImporter& SceneImporter);
 
-	int32 ProcessMaterialElement(TSharedPtr< IDatasmithMasterMaterialElement > BaseMaterialElement, const TCHAR* Host, FTextureCallback TextureCallback);
+	int32 ProcessMaterialElement(TSharedPtr< IDatasmithMasterMaterialElement > BaseMaterialElement, FTextureCallback TextureCallback);
 
 	int32 ProcessMaterialElement(IDatasmithUEPbrMaterialElement* PbrMaterialElement, FTextureCallback TextureCallback);
 
diff --git a/Engine/Plugins/Experimental/Enterprise/DatasmithRuntime/Source/Private/MaterialImporter.cpp b/Engine/Plugins/Experimental/Enterprise/DatasmithRuntime/Source/Private/MaterialImporter.cpp
--- a/Engine/Plugins/Experimental/Enterprise/DatasmithRuntime/Source/Private/MaterialImporter.cpp
+++ b/Engine/Plugins/Experimental/Enterprise/DatasmithRuntime/Source/Private/MaterialImporter.cpp
@@ -105,7 +105,7 @@ namespace DatasmithRuntime
 		}
 		else if( Element->IsA( EDatasmithElementType::MasterMaterial ) )
 		{
-			MaterialData.Requirements = ProcessMaterialElement(StaticCastSharedPtr<IDatasmithMasterMaterialElement>(Element), *Host, TextureCallback);
+			MaterialData.Requirements = ProcessMaterialElement(StaticCastSharedPtr<IDatasmithMasterMaterialElement>(Element), TextureCallback);
 		}
 
 		MaterialData.SetState(EAssetState::Processed);
diff --git a/Engine/Plugins/Experimental/Enterprise/DatasmithRuntime/Source/Private/SceneImporter.cpp b/Engine/Plugins/Experimental/Enterprise/DatasmithRuntime/Source/Private/SceneImporter.cpp
--- a/Engine/Plugins/Experimental/Enterprise/DatasmithRuntime/Source/Private/SceneImporter.cpp
+++ b/Engine/Plugins/Experimental/Enterprise/DatasmithRuntime/Source/Private/SceneImporter.cpp
@@ -4,8 +4,10 @@
 
 #include "DatasmithRuntimeUtils.h"
 #include "LogCategory.h"
+#include "MaterialImportUtils.h"
 
 #include "DatasmithNativeTranslator.h"
+#include "DatasmithMaterialElements.h"
 #include "IDatasmithSceneElements.h"
 
 #include "Camera/PlayerCameraManager.h"
@@ -631,59 +633,7 @@ namespace DatasmithRuntime
 
 		IncrementalModifications(UpdateContext.Updates);
 
-		if (UpdateContext.Deletions.Num() > 0)
-		{
-			// Terminate all rendering commands before deleting any component and asset
-			FlushRenderingCommands();
-
-			FActionTaskFunction TaskFunc = [this](UObject*, const FReferencer& Referencer) -> EActionResult::Type
-			{
-				EActionResult::Type Result = this->DeleteElement(Referencer.GetId());
-
-				if(Result == EActionResult::Succeeded)
-				{
-					this->TasksToComplete |= EWorkerTask::GarbageCollect;
-				}
-
-				return Result;
-			};
-
-			for (DirectLink::FSceneGraphId& ElementId : UpdateContext.Deletions)
-			{
-				if (Elements.Contains(ElementId))
-				{
-					if (AssetDataList.Contains(ElementId))
-					{
-						if (!AssetDataList[ElementId].HasState(EAssetState::PendingDelete))
-						{
-							continue;
-						}
-
-						const EDataType DataType = AssetDataList[ElementId].Type;
-						const FString& Prefix = DataType == EDataType::Texture ? TexturePrefix : (DataType == EDataType::Material ? MaterialPrefix : MeshPrefix);
-
-						UE_LOG(LogDatasmithRuntime, Log, TEXT("IncrementalDeletions: %s %d"), *Prefix, ElementId);
-
-						AddToQueue(EQueueTask::DeleteAssetQueue, { TaskFunc, FReferencer(ElementId) } );
-						TasksToComplete |= EWorkerTask::DeleteAsset;
-					}
-					else if (ActorDataList.Contains(ElementId))
-					{
-						UE_LOG(LogDatasmithRuntime, Log, TEXT("IncrementalDeletions: actor %s (%d)"), Elements[ElementId]->GetLabel(), ElementId);
-
-						HideSceneComponent(ActorDataList[ElementId].GetObject<USceneComponent>());
-
-						AddToQueue(EQueueTask::DeleteCompQueue, { TaskFunc, FReferencer(ElementId) } );
-						TasksToComplete |= EWorkerTask::DeleteComponent;
-					}
-					else
-					{
-						UE_LOG(LogDatasmithRuntime, Error, TEXT("Element %d (%s) was not found"), ElementId, Elements[ElementId]->GetName());
-						ensure(false);
-					}
-				}
-			}
-		}
+		IncrementalDeletions(UpdateContext.Deletions);
 
 		bIncrementalUpdate = true;
 
@@ -725,25 +675,23 @@ namespace DatasmithRuntime
 					const EDataType DataType = AssetDataList[ElementId].Type;
 					const FString& Prefix = DataType == EDataType::Texture ? TexturePrefix : (DataType == EDataType::Material ? MaterialPrefix : MeshPrefix);
 
-					UE_LOG(LogDatasmithRuntime, Log, TEXT("IncrementalModifications: %s %d"), *Prefix, ElementPtr->GetNodeId());
+					UE_LOG(LogDatasmithRuntime, Log, TEXT("IncrementalModifications: %s %s (%d)"), *Prefix, ElementPtr->GetName(), ElementPtr->GetNodeId());
 
 					const FString PrefixedName = Prefix + ElementPtr->GetName();
 
 					if (!AssetElementMapping.Contains(PrefixedName))
 					{
-						AssetElementMapping.Add(PrefixedName, ElementId);
-
-						FString OldKey;
 						for (TPair<FString, FSceneGraphId>& Entry : AssetElementMapping)
 						{
 							if (Entry.Value == ElementId)
 							{
-								OldKey = Entry.Key;
+								const FString OldKey = Entry.Key;
+								AssetElementMapping.Remove(OldKey);
 								break;
 							}
 						}
 
-						AssetElementMapping.Remove(OldKey);
+						AssetElementMapping.Add(PrefixedName, ElementId);
 					}
 
 					FActionTaskFunction TaskFunc;
@@ -804,6 +752,8 @@ namespace DatasmithRuntime
 				{
 					FActorData& ActorData = ActorDataList[ElementId];
 
+					UE_LOG(LogDatasmithRuntime, Log, TEXT("IncrementalModifications: Actor %s (%d)"), ElementPtr->GetName(), ElementPtr->GetNodeId());
+
 					ActorData.SetState(EAssetState::Unknown);
 				}
 			}
@@ -814,6 +764,127 @@ namespace DatasmithRuntime
 		}
 	}
 
+	void FSceneImporter::IncrementalDeletions(TArray<DirectLink::FSceneGraphId>& Deletions)
+	{
+		if (Deletions.Num() == 0)
+		{
+			return;
+		}
+
+		FActionTaskFunction TaskFunc = [this](UObject*, const FReferencer& Referencer) -> EActionResult::Type
+		{
+			EActionResult::Type Result = this->DeleteElement(Referencer.GetId());
+
+			if(Result == EActionResult::Succeeded)
+			{
+				this->TasksToComplete |= EWorkerTask::GarbageCollect;
+			}
+
+			return Result;
+		};
+
+		TFunction<void(FSceneGraphId*, FSceneGraphId)> RemoveFromReferencer;
+		RemoveFromReferencer = [this](FSceneGraphId* AssetIdPtr, FSceneGraphId ElementId) -> void
+		{
+			if (AssetIdPtr)
+			{
+				TArray<FReferencer>& Referencers = this->AssetDataList[*AssetIdPtr].Referencers;
+
+				for (int32 Index = 0; Index < Referencers.Num(); ++Index)
+				{
+					if (Referencers[Index].ElementId == ElementId)
+					{
+						Referencers.RemoveAt(Index, 1, false);
+						break;
+					}
+				}
+			}
+		};
+
+		for (DirectLink::FSceneGraphId& ElementId : Deletions)
+		{
+			if (Elements.Contains(ElementId))
+			{
+				if (AssetDataList.Contains(ElementId))
+				{
+					if (!AssetDataList[ElementId].HasState(EAssetState::PendingDelete))
+					{
+						continue;
+					}
+
+					const EDataType DataType = AssetDataList[ElementId].Type;
+					const FString& Prefix = DataType == EDataType::Texture ? TexturePrefix : (DataType == EDataType::Material ? MaterialPrefix : MeshPrefix);
+
+					UE_LOG(LogDatasmithRuntime, Log, TEXT("IncrementalDeletions: %s %d"), *Prefix, ElementId);
+
+					// If asset is a material and it references textures, remove it from the textures' list of referencers
+					if (DataType == EDataType::Material)
+					{
+						FTextureCallback TextureCallback;
+						TextureCallback = [this, ElementId, RemoveFromReferencer](const FString& TextureNamePrefixed, int32 PropertyIndex)->void
+						{
+							RemoveFromReferencer(this->AssetElementMapping.Find(TextureNamePrefixed), ElementId);
+						};
+
+						TSharedPtr< IDatasmithElement >& Element = Elements[ ElementId ];
+
+						if( Element->IsA( EDatasmithElementType::UEPbrMaterial ) )
+						{
+							ProcessMaterialElement(static_cast<IDatasmithUEPbrMaterialElement*>(Element.Get()), TextureCallback);
+						}
+						else if( Element->IsA( EDatasmithElementType::MasterMaterial ) )
+						{
+							ProcessMaterialElement(StaticCastSharedPtr<IDatasmithMasterMaterialElement>(Element), TextureCallback);
+						}
+					}
+					// If asset is a mesh and it references materials, remove it from the materials' list of referencers
+					else if (DataType == EDataType::Mesh)
+					{
+						TSharedPtr< IDatasmithMeshElement > MeshElement = StaticCastSharedPtr< IDatasmithMeshElement >(Elements[ElementId]);
+
+						for (int32 Index = 0; Index < MeshElement->GetMaterialSlotCount(); Index++)
+						{
+							if (const IDatasmithMaterialIDElement* MaterialIDElement = MeshElement->GetMaterialSlotAt(Index).Get())
+							{
+								const FString MaterialPathName(MaterialIDElement->GetName());
+
+								if (!MaterialPathName.StartsWith(TEXT("/")))
+								{
+									RemoveFromReferencer(AssetElementMapping.Find(MaterialPrefix + MaterialPathName), ElementId);
+								}
+							}
+						}
+					}
+
+					AddToQueue(EQueueTask::DeleteAssetQueue, { TaskFunc, FReferencer(ElementId) } );
+					TasksToComplete |= EWorkerTask::DeleteAsset;
+				}
+				else if (ActorDataList.Contains(ElementId))
+				{
+					UE_LOG(LogDatasmithRuntime, Log, TEXT("IncrementalDeletions: actor %s (%d)"), Elements[ElementId]->GetLabel(), ElementId);
+
+					FActorData& ActorData = ActorDataList[ElementId];
+
+					HideSceneComponent(ActorData.GetObject<USceneComponent>());
+
+					// Remove actor from the asset's list of referencers
+					if (ActorData.AssetId != INDEX_NONE && AssetDataList.Contains(ActorData.AssetId))
+					{
+						RemoveFromReferencer(&AssetDataList[ActorData.AssetId].ElementId, ElementId);
+					}
+
+					AddToQueue(EQueueTask::DeleteCompQueue, { TaskFunc, FReferencer(ElementId) } );
+					TasksToComplete |= EWorkerTask::DeleteComponent;
+				}
+				else
+				{
+					UE_LOG(LogDatasmithRuntime, Error, TEXT("Element %d (%s) was not found"), ElementId, Elements[ElementId]->GetName());
+					ensure(false);
+				}
+			}
+		}
+	}
+
 	void FSceneImporter::ProcessDependency(const TSharedPtr<IDatasmithElement>& Element)
 	{
 		FSceneGraphId ElementId = Element->GetNodeId();
@@ -912,12 +983,12 @@ namespace DatasmithRuntime
 		{
 			if (ElementPtr->IsA(EDatasmithElementType::BaseMaterial))
 			{
-				UE_LOG(LogDatasmithRuntime, Log, TEXT("IncrementalAdditions: Material %d"), ElementPtr->GetNodeId());
+				UE_LOG(LogDatasmithRuntime, Log, TEXT("IncrementalAdditions: Material %s (%d)"), ElementPtr->GetName(), ElementPtr->GetNodeId());
 				LocalAddAsset(MoveTemp(ElementPtr), EDataType::Material);
 			}
 			else if (ElementPtr->IsA(EDatasmithElementType::StaticMesh))
 			{
-				UE_LOG(LogDatasmithRuntime, Log, TEXT("IncrementalAdditions: StaticMesh %d"), ElementPtr->GetNodeId());
+				UE_LOG(LogDatasmithRuntime, Log, TEXT("IncrementalAdditions: StaticMesh %s (%d)"), ElementPtr->GetName(), ElementPtr->GetNodeId());
 				if (IDatasmithMeshElement* MeshElement = static_cast<IDatasmithMeshElement*>(ElementPtr.Get()))
 				{
 					// If resource file does not exist, add scene's resource path if valid
@@ -935,7 +1006,7 @@ namespace DatasmithRuntime
 			}
 			else if (ElementPtr->IsA(EDatasmithElementType::Texture))
 			{
-				UE_LOG(LogDatasmithRuntime, Log, TEXT("IncrementalAdditions: Texture %d"), ElementPtr->GetNodeId());
+				UE_LOG(LogDatasmithRuntime, Log, TEXT("IncrementalAdditions: Texture %s (%d)"), ElementPtr->GetName(), ElementPtr->GetNodeId());
 				if (IDatasmithTextureElement* TextureElement = static_cast<IDatasmithTextureElement*>(ElementPtr.Get()))
 				{
 					// If resource file does not exist, add scene's resource path if valid
@@ -950,6 +1021,10 @@ namespace DatasmithRuntime
 					}
 				}
 			}
+			else if (ElementPtr->IsA(EDatasmithElementType::Actor))
+			{
+				UE_LOG(LogDatasmithRuntime, Log, TEXT("IncrementalAdditions: Actor %s (%d)"), ElementPtr->GetName(), ElementPtr->GetNodeId());
+			}
 		}
 
 		TasksToComplete |= EWorkerTask::SetupTasks;
diff --git a/Engine/Plugins/Experimental/Enterprise/DatasmithRuntime/Source/Private/SceneImporter.h b/Engine/Plugins/Experimental/Enterprise/DatasmithRuntime/Source/Private/SceneImporter.h
--- a/Engine/Plugins/Experimental/Enterprise/DatasmithRuntime/Source/Private/SceneImporter.h
+++ b/Engine/Plugins/Experimental/Enterprise/DatasmithRuntime/Source/Private/SceneImporter.h
@@ -528,6 +528,8 @@ namespace DatasmithRuntime
 
 		void IncrementalModifications(TArray<TSharedPtr<IDatasmithElement>>& Modifications);
 
+		void IncrementalDeletions(TArray<DirectLink::FSceneGraphId>& Deletions);
+
 		/** Add an FAssetData object associated with the element's id to the map */
 		void AddAsset(TSharedPtr<IDatasmithElement>&& ElementPtr, const FString& Prefix, EDataType InType = EDataType::None);
 
