Added support for metadata
Added tracking of property change for material and metadata
Removed boolean to update gamer player viewpoint to the first imported camera

diff --git a/Engine/Plugins/Experimental/Enterprise/DatasmithRuntime/Source/Private/ActorImporter.cpp b/Engine/Plugins/Experimental/Enterprise/DatasmithRuntime/Source/Private/ActorImporter.cpp
new file mode 100644
--- /dev/null
+++ b/Engine/Plugins/Experimental/Enterprise/DatasmithRuntime/Source/Private/ActorImporter.cpp
@@ -0,0 +1,553 @@
+// Copyright Epic Games, Inc. All Rights Reserved.
+
+#include "SceneImporter.h"
+
+#include "DatasmithRuntimeUtils.h"
+#include "LogCategory.h"
+
+#include "IDatasmithSceneElements.h"
+
+#include "Algo/AnyOf.h"
+#include "Camera/PlayerCameraManager.h"
+#include "CineCameraActor.h"
+#include "CineCameraComponent.h"
+#include "Components/SceneComponent.h"
+#include "Engine/World.h"
+#include "GameFramework/Actor.h"
+#include "ProfilingDebugging/MiscTrace.h"
+
+namespace DatasmithRuntime
+{
+	// Tag added to created scene components and actors
+	// Used during reset process
+	const FName RuntimeTag("Datasmith.Runtime.Tag");
+
+	void DeleteSceneComponent(USceneComponent* SceneComponent);
+
+	// Recursively destroy actor from world. Children are deleted first.
+	void DeleteActorRecursive(UWorld* GameWorld, AActor* Actor);
+
+	// Set scene component's visibility to false and remove geometry from static mesh component
+	// Call before deletion to ensure no asset is used
+	void HideSceneComponent(USceneComponent* SceneComponent);
+
+	void FSceneImporter::ProcessActorElement(const TSharedPtr< IDatasmithActorElement >& ActorElement, FSceneGraphId ParentId)
+	{
+		TRACE_CPUPROFILER_EVENT_SCOPE(FSceneImporter::ProcessActorElement);
+
+		FActorData& ActorData = FindOrAddActorData(ActorElement);
+
+		if (ActorData.HasState(EAssetState::Processed))
+		{
+			return;
+		}
+
+		ActorData.ParentId = ParentId;
+		ActorData.WorldTransform = FTransform( ActorElement->GetRotation(), ActorElement->GetTranslation(), ActorElement->GetScale() );
+
+		bool bProcessSuccessful = false;
+
+		if (ActorElement->IsA(EDatasmithElementType::StaticMeshActor))
+		{
+			IDatasmithMeshActorElement* MeshActorElement = static_cast<IDatasmithMeshActorElement*>(ActorElement.Get());
+
+			ActorData.Type = EDataType::MeshActor;
+			bProcessSuccessful = ProcessMeshActorData(ActorData, MeshActorElement);
+		}
+		else if (ActorElement->IsA(EDatasmithElementType::Light))
+		{
+			IDatasmithLightActorElement* LightActorElement = static_cast<IDatasmithLightActorElement*>(ActorElement.Get());
+
+			ActorData.Type = EDataType::LightActor;
+			bProcessSuccessful = ProcessLightActorData(ActorData, LightActorElement);
+		}
+		else if (ActorElement->IsA(EDatasmithElementType::Camera))
+		{
+			IDatasmithCameraActorElement* CameraElement = static_cast<IDatasmithCameraActorElement*>(ActorElement.Get());
+
+			bProcessSuccessful = ProcessCameraActorData(ActorData, CameraElement);
+		}
+
+		if (!bProcessSuccessful)
+		{
+			CreateActorComponent(ActorData, ActorElement);
+			ActorData.AddState(EAssetState::Processed | EAssetState::Completed);
+		}
+	}
+
+	bool FSceneImporter::DeleteComponent(FActorData& ActorData)
+	{
+		if (USceneComponent* SceneComponent = ActorData.GetObject<USceneComponent>())
+		{
+			if (ImportOptions.BuildHierarchy !=  EBuildHierarchyMethod::None)
+			{
+				USceneComponent* ParentComponent = SceneComponent->GetAttachmentRoot();
+
+				AActor* Actor = SceneComponent->GetOwner();
+
+				if (ParentComponent->GetOwner() == Actor)
+				{
+					TArray<USceneComponent*> Children = SceneComponent->GetAttachChildren();
+
+					for (USceneComponent* ChildComponent : Children)
+					{
+						ChildComponent->DetachFromComponent(FDetachmentTransformRules::KeepRelativeTransform);
+						ChildComponent->AttachToComponent(ParentComponent, FAttachmentTransformRules::KeepRelativeTransform);
+					}
+
+					DeleteSceneComponent(SceneComponent);
+				}
+				else
+				{
+					AActor* ParentActor = ParentComponent->GetOwner();
+
+					TArray<AActor*> ChildActors;
+					Actor->GetAttachedActors(ChildActors);
+
+					for (AActor* ChildActor : ChildActors)
+					{
+						ChildActor->DetachFromActor(FDetachmentTransformRules::KeepRelativeTransform);
+						ChildActor->AttachToActor(ParentActor, FAttachmentTransformRules::KeepRelativeTransform);
+					}
+
+					RootComponent->GetOwner()->GetWorld()->EditorDestroyActor(Actor, true);
+				}
+			}
+			else if (SceneComponent->GetAttachmentRoot() == RootComponent.Get())
+			{
+				DeleteSceneComponent(SceneComponent);
+			}
+
+			ActorData.Object.Reset();
+
+			return true;
+		}
+
+		return false;
+	}
+
+	bool FSceneImporter::DeleteData()
+	{
+		// Terminate all rendering commands before deleting any component and asset
+		FlushRenderingCommands();
+
+		bool bGarbageCollect = false;
+
+		TArray<AActor*> ChildActors;
+		RootComponent->GetOwner()->GetAttachedActors(ChildActors);
+
+		// Check to see if last import required to build the actor's hierarchy
+		bool bHasHierarchy = Algo::AnyOf(ChildActors, [](const AActor* ChildActor){ return ChildActor->ActorHasTag(RuntimeTag); });
+
+		// If last import created actors, delete them from bottom to top
+		if (bHasHierarchy)
+		{
+			// Remove hold onto scene component as they will be deleted along with the owning actor.
+			for (TPair< FSceneGraphId, FActorData >& Pair : ActorDataList)
+			{
+				Pair.Value.Object.Reset();
+			}
+
+			UWorld* GameWorld = RootComponent->GetOwner()->GetWorld();
+
+			for (AActor* ChildActor : ChildActors)
+			{
+				// Only delete actors created by runtime import
+				if (ChildActor->ActorHasTag(RuntimeTag))
+				{
+					ChildActor->DetachFromActor(FDetachmentTransformRules::KeepRelativeTransform);
+					DeleteActorRecursive(GameWorld, ChildActor);
+				}
+			}
+
+			TArray<USceneComponent*> SceneComponents = RootComponent->GetAttachChildren();
+			for (USceneComponent* SceneComponent : SceneComponents)
+			{
+				// Only delete components created by runtime import
+				if (SceneComponent->ComponentHasTag(RuntimeTag))
+				{
+					DeleteSceneComponent(SceneComponent);
+				}
+			}
+
+			bGarbageCollect = true;
+		}
+		else
+		{
+			for (TPair< FSceneGraphId, FActorData >& Pair : ActorDataList)
+			{
+				bGarbageCollect |= DeleteComponent(Pair.Value);
+			}
+		}
+
+		for (TPair< FSceneGraphId, FAssetData >& Entry : AssetDataList)
+		{
+			bGarbageCollect |= DeleteAsset(Entry.Value);
+		}
+
+		return bGarbageCollect;
+	}
+
+	bool FSceneImporter::ProcessCameraActorData(FActorData& ActorData, IDatasmithCameraActorElement* CameraElement)
+	{
+		if (ActorData.HasState(EAssetState::Processed) || ActorData.Object.IsValid())
+		{
+			ActorData.AddState(EAssetState::Processed);
+			return true;
+		}
+
+		if (ImportOptions.BuildHierarchy != EBuildHierarchyMethod::None)
+		{
+			UCineCameraComponent* CameraComponent = ActorData.GetObject<UCineCameraComponent>();
+
+			if (CameraComponent == nullptr)
+			{
+				if (ImportOptions.BuildHierarchy != EBuildHierarchyMethod::None)
+				{
+					UWorld* World = RootComponent->GetOwner()->GetWorld();
+
+					ACineCameraActor* CameraActor = Cast< ACineCameraActor >( World->SpawnActor( ACineCameraActor::StaticClass(), nullptr, nullptr ) );
+					CameraActor->Rename(CameraElement->GetName(), nullptr, REN_NonTransactional | REN_DontCreateRedirectors);
+#if WITH_EDITOR
+					CameraActor->SetActorLabel(CameraElement->GetLabel());
+#endif
+					CameraComponent = CameraActor->GetCineCameraComponent();
+				}
+				else
+				{
+					FName ComponentName = NAME_None;
+					ComponentName = MakeUniqueObjectName(RootComponent->GetOwner(), UCineCameraComponent::StaticClass(), CameraElement->GetLabel());
+					CameraComponent = NewObject< UCineCameraComponent >(RootComponent->GetOwner(), ComponentName);
+				}
+
+				ActorData.Object = CameraComponent;
+			}
+			//else
+			//{
+			//	CameraComponent->GetOwner()->UpdateComponentTransforms();
+			//	CameraComponent->GetOwner()->MarkComponentsRenderStateDirty();
+			//}
+
+			CameraComponent->Filmback.SensorWidth = CameraElement->GetSensorWidth();
+			CameraComponent->Filmback.SensorHeight = CameraElement->GetSensorWidth() / CameraElement->GetSensorAspectRatio();
+			CameraComponent->LensSettings.MaxFStop = 32.0f;
+			CameraComponent->CurrentFocalLength = CameraElement->GetFocalLength();
+			CameraComponent->CurrentAperture = CameraElement->GetFStop();
+
+			CameraComponent->FocusSettings.FocusMethod = CameraElement->GetEnableDepthOfField() ? ECameraFocusMethod::Manual : ECameraFocusMethod::DoNotOverride;
+			CameraComponent->FocusSettings.ManualFocusDistance = CameraElement->GetFocusDistance();
+
+			if (const IDatasmithPostProcessElement* PostProcess = CameraElement->GetPostProcess().Get())
+			{
+				FPostProcessSettings& PostProcessSettings = CameraComponent->PostProcessSettings;
+
+				if ( !FMath::IsNearlyEqual( PostProcess->GetTemperature(), 6500.f ) )
+				{
+					PostProcessSettings.bOverride_WhiteTemp = true;
+					PostProcessSettings.WhiteTemp = PostProcess->GetTemperature();
+				}
+
+				if ( PostProcess->GetVignette() > 0.f )
+				{
+					PostProcessSettings.bOverride_VignetteIntensity = true;
+					PostProcessSettings.VignetteIntensity = PostProcess->GetVignette();
+				}
+
+				if (PostProcess->GetColorFilter() != FLinearColor::Black && PostProcess->GetColorFilter() != FLinearColor::White )
+				{
+					PostProcessSettings.bOverride_FilmWhitePoint = true;
+					PostProcessSettings.FilmWhitePoint = PostProcess->GetColorFilter();
+				}
+
+				if ( !FMath::IsNearlyEqual( PostProcess->GetSaturation(), 1.f ) )
+				{
+					PostProcessSettings.bOverride_ColorSaturation = true;
+					PostProcessSettings.ColorSaturation.W = PostProcess->GetSaturation();
+				}
+
+				if ( PostProcess->GetCameraISO() > 0.f || PostProcess->GetCameraShutterSpeed() > 0.f || PostProcess->GetDepthOfFieldFstop() > 0.f )
+				{
+					PostProcessSettings.bOverride_AutoExposureMethod = true;
+					PostProcessSettings.AutoExposureMethod = EAutoExposureMethod::AEM_Manual;
+
+					if ( PostProcess->GetCameraISO() > 0.f )
+					{
+						PostProcessSettings.bOverride_CameraISO = true;
+						PostProcessSettings.CameraISO = PostProcess->GetCameraISO();
+					}
+
+					if ( PostProcess->GetCameraShutterSpeed() > 0.f )
+					{
+						PostProcessSettings.bOverride_CameraShutterSpeed = true;
+						PostProcessSettings.CameraShutterSpeed = PostProcess->GetCameraShutterSpeed();
+					}
+
+					if ( PostProcess->GetDepthOfFieldFstop() > 0.f )
+					{
+						PostProcessSettings.bOverride_DepthOfFieldFstop = true;
+						PostProcessSettings.DepthOfFieldFstop = PostProcess->GetDepthOfFieldFstop();
+					}
+				}
+			}
+		}
+
+		FinalizeComponent(ActorData);
+
+		ActorData.SetState(EAssetState::Processed | EAssetState::Completed);
+
+		return true;
+	}
+
+	FActorData& FSceneImporter::FindOrAddActorData(const TSharedPtr<IDatasmithActorElement>& ActorElement)
+	{
+		FSceneGraphId ElementId = ActorElement->GetNodeId();
+
+		if (!Elements.Contains(ElementId))
+		{
+			Elements.Add(ElementId, ActorElement);
+
+			FActorData ActorData(ElementId);
+			return ActorDataList.Emplace(ElementId, MoveTemp(ActorData));
+		}
+
+		ensure(ActorDataList.Contains(ElementId));
+		return ActorDataList[ElementId];
+	}
+
+	void FSceneImporter::CreateActorComponent(FActorData& ActorData, const TSharedPtr< IDatasmithActorElement >& ActorElement)
+	{
+		const bool bBuildHierarchy = ImportOptions.BuildHierarchy == EBuildHierarchyMethod::Unfiltered || (ImportOptions.BuildHierarchy == EBuildHierarchyMethod::Simplified && ActorElement->GetChildrenCount() > 0);
+		if (!bBuildHierarchy)
+		{
+			ActorData.SetState(EAssetState::Skipped);
+			return;
+		}
+
+		bool bMustProcess = true;
+
+		if (ImportOptions.BuildHierarchy == EBuildHierarchyMethod::Simplified)
+		{
+			TFunction<bool(const TSharedPtr< IDatasmithActorElement >&)> MustKeep;
+			MustKeep = [this, &MustKeep](const TSharedPtr< IDatasmithActorElement >& ActorElementToCheck) -> bool
+			{
+				for (int32 Index = 0; Index < ActorElementToCheck->GetChildrenCount(); ++Index)
+				{
+					if (MustKeep(ActorElementToCheck->GetChild(Index)))
+					{
+						return true;
+					}
+				}
+
+				const bool bMustKeep = ActorElementToCheck->IsA(EDatasmithElementType::Camera)
+					|| ActorElementToCheck->IsA(EDatasmithElementType::Light)
+					|| ActorElementToCheck->IsA(EDatasmithElementType::StaticMeshActor);
+
+				if (!bMustKeep)
+				{
+					FActorData& LocalActorData = this->FindOrAddActorData(ActorElementToCheck);
+					LocalActorData.SetState(EAssetState::Processed | EAssetState::Completed | EAssetState::Skipped);
+				}
+
+				return bMustKeep;
+			};
+
+			bMustProcess = MustKeep(ActorElement);
+
+			if (bMustProcess && ActorElement->GetChildrenCount() == 1 && ActorElement->GetChild(0)->GetChildrenCount() == 0)
+			{
+				ActorElement->GetChild(0)->SetIsAComponent(false);
+				ActorData.SetState(EAssetState::Skipped);
+				bMustProcess = false;
+			}
+		}
+
+		if(bMustProcess)
+		{
+			USceneComponent* SceneComponent = ActorData.GetObject<USceneComponent>();
+
+			if (SceneComponent == nullptr)
+			{
+				AActor* Actor = RootComponent->GetOwner()->GetWorld()->SpawnActor(AActor::StaticClass(), nullptr, nullptr);
+				Actor->Rename(ActorElement->GetName(), nullptr, REN_NonTransactional | REN_DontCreateRedirectors);
+#if WITH_EDITOR
+				Actor->SetActorLabel( ActorElement->GetLabel() );
+#endif
+
+				Actor->Tags.Empty(ActorElement->GetTagsCount());
+				for (int32 Index = 0; Index < ActorElement->GetTagsCount(); ++Index)
+				{
+					Actor->Tags.Add(ActorElement->GetTag(Index));
+				}
+
+				// #ue_datasmithruntime: Add Layers and metadata
+
+				SceneComponent = Actor->GetRootComponent();
+
+				if ( !SceneComponent )
+				{
+					SceneComponent = NewObject< USceneComponent >( Actor, USceneComponent::StaticClass(), ActorElement->GetLabel(), RF_Transactional );
+
+					Actor->AddInstanceComponent( SceneComponent );
+					Actor->SetRootComponent( SceneComponent );
+				}
+
+				ActorData.Object = TWeakObjectPtr<UObject>(SceneComponent);
+			}
+
+			FinalizeComponent(ActorData);
+		}
+	}
+
+	void FSceneImporter::FinalizeComponent(FActorData& ActorData)
+	{
+		USceneComponent* SceneComponent = ActorData.GetObject<USceneComponent>();
+		if (!SceneComponent)
+		{
+			return;
+		}
+
+		if (SceneComponent->IsRegistered())
+		{
+			SceneComponent->UnregisterComponent();
+		}
+
+		// Find the right parent
+		FSceneGraphId ParentId = ActorData.ParentId;
+		while (ParentId != DirectLink::InvalidId && ActorDataList[ParentId].HasState(EAssetState::Skipped))
+		{
+			ParentId = ActorDataList[ParentId].ParentId;
+		}
+
+		if (ImportOptions.BuildHierarchy != EBuildHierarchyMethod::None && ParentId != DirectLink::InvalidId)
+		{
+			// If hierarchy building is required, wait for parent component to be created if applicable
+			if (ActorDataList[ParentId].GetObject<USceneComponent>() == nullptr)
+			{
+				FActionTaskFunction FinalizeComponentFunc = [this,ParentId](UObject* Object, const FReferencer& Referencer) -> EActionResult::Type
+				{
+					FActorData& ActorData = ActorDataList[Referencer.GetId()];
+
+					if (ActorDataList[ParentId].GetObject<USceneComponent>() == nullptr)
+					{
+						return EActionResult::Retry;
+					}
+
+					this->FinalizeComponent(ActorData);
+
+					return EActionResult::Succeeded;
+				};
+
+				AddToQueue(EQueueTask::NonAsyncQueue, { FinalizeComponentFunc, { EDataType::Actor, ActorData.ElementId, 0 } });
+				TasksToComplete |= EWorkerTask::ComponentFinalize;
+
+				return;
+			}
+		}
+
+		const bool bUseParent = ParentId != DirectLink::InvalidId && ImportOptions.BuildHierarchy != EBuildHierarchyMethod::None;
+		USceneComponent* ParentComponent = bUseParent ? ActorDataList[ParentId].GetObject<USceneComponent>() : RootComponent.Get();
+
+		if (ParentComponent)
+		{
+			const IDatasmithActorElement* ActorElement = static_cast<IDatasmithActorElement*>(Elements[ActorData.ElementId].Get());
+
+			const FTransform ParentToWorld = ParentComponent->GetComponentToWorld();
+			const FTransform RelativeTransform = ActorData.WorldTransform.GetRelativeTransform(ParentToWorld);
+
+			SceneComponent->SetRelativeTransform( RelativeTransform );
+			SceneComponent->SetVisibility(ActorElement->GetVisibility());
+			SceneComponent->SetMobility( EComponentMobility::Movable );
+
+			if (SceneComponent->GetAttachParent() != ParentComponent)
+			{
+				SceneComponent->DetachFromComponent(FDetachmentTransformRules::KeepRelativeTransform);
+				SceneComponent->AttachToComponent( ParentComponent, FAttachmentTransformRules::KeepRelativeTransform );
+			}
+
+			SceneComponent->ComponentTags.Empty(ActorElement->GetTagsCount());
+			for (int32 Index = 0; Index < ActorElement->GetTagsCount(); ++Index)
+			{
+				SceneComponent->ComponentTags.Add(ActorElement->GetTag(Index));
+			}
+
+			// Add runtime tag if scene component is attached to runtime actor's component
+			if (ParentComponent == RootComponent)
+			{
+				SceneComponent->ComponentTags.Add(RuntimeTag);
+			}
+
+			if (ImportOptions.BuildHierarchy != EBuildHierarchyMethod::None)
+			{
+				if (ActorElement->IsAComponent() && ParentComponent && SceneComponent->GetOwner() != ParentComponent->GetOwner())
+				{
+					FName UniqueName = MakeUniqueObjectName(ParentComponent->GetOwner(), SceneComponent->GetClass(), ActorElement->GetName());
+					SceneComponent->Rename(*UniqueName.ToString(), ParentComponent->GetOwner(), REN_NonTransactional | REN_DontCreateRedirectors);
+				}
+
+				SceneComponent->GetOwner()->AddInstanceComponent(SceneComponent);
+
+				// Add runtime tag on actor created on import.
+				if (SceneComponent->GetOwner() != RootComponent->GetOwner() && !SceneComponent->GetOwner()->ActorHasTag(RuntimeTag))
+				{
+					SceneComponent->GetOwner()->Tags.Add(RuntimeTag);
+				}
+			}
+
+			SceneComponent->RegisterComponentWithWorld(RootComponent->GetOwner()->GetWorld());
+		}
+		else
+		{
+			ensure(false);
+		}
+
+		ApplyMetadata(ActorData.MetadataId, SceneComponent);
+	}
+
+	void DeleteActorRecursive(UWorld* GameWorld, AActor* Actor)
+	{
+		TArray<AActor*> ChildActors;
+		Actor->GetAttachedActors(ChildActors);
+
+		for (AActor* ChildActor : ChildActors)
+		{
+			ChildActor->DetachFromActor(FDetachmentTransformRules::KeepRelativeTransform);
+			DeleteActorRecursive(GameWorld, ChildActor);
+		}
+
+		ensure(GameWorld->DestroyActor(Actor));
+
+		// Since deletion can be delayed, rename to avoid future name collision
+		// Call UObject::Rename directly on actor to avoid AActor::Rename which unnecessarily sunregister and re-register components
+		Actor->UObject::Rename( nullptr, GetTransientPackage(), REN_DontCreateRedirectors | REN_ForceNoResetLoaders );
+	}
+
+	void HideSceneComponent(USceneComponent* SceneComponent)
+	{
+		if (SceneComponent)
+		{
+			SceneComponent->SetVisibility(false);
+
+			if (UStaticMeshComponent* MeshComponent = Cast< UStaticMeshComponent >(SceneComponent))
+			{
+				MeshComponent->OverrideMaterials.Reset();
+				MeshComponent->SetStaticMesh(nullptr);
+				MeshComponent->ReleaseResources();
+			}
+
+			SceneComponent->MarkRenderStateDirty();
+			SceneComponent->MarkRenderDynamicDataDirty();
+		}
+	}
+
+	void DeleteSceneComponent(USceneComponent* SceneComponent)
+	{
+		SceneComponent->UnregisterComponent();
+
+		SceneComponent->DetachFromComponent(FDetachmentTransformRules::KeepRelativeTransform);
+
+		SceneComponent->ClearFlags(RF_AllFlags);
+		SceneComponent->SetFlags(RF_Transient);
+		SceneComponent->Rename(nullptr, nullptr, REN_NonTransactional | REN_DontCreateRedirectors);
+		SceneComponent->MarkPendingKill();
+	}
+} // End of namespace DatasmithRuntime
\ No newline at end of file
diff --git a/Engine/Plugins/Experimental/Enterprise/DatasmithRuntime/Source/Private/MaterialImportUtils.cpp b/Engine/Plugins/Experimental/Enterprise/DatasmithRuntime/Source/Private/MaterialImportUtils.cpp
--- a/Engine/Plugins/Experimental/Enterprise/DatasmithRuntime/Source/Private/MaterialImportUtils.cpp
+++ b/Engine/Plugins/Experimental/Enterprise/DatasmithRuntime/Source/Private/MaterialImportUtils.cpp
@@ -221,7 +221,7 @@ namespace DatasmithRuntime
 
 		for (int Index = 0; Index < MaterialElement->GetPropertiesCount(); ++Index)
 		{
-			const TSharedPtr< IDatasmithKeyValueProperty > Property = MaterialElement->GetProperty(Index);
+			const TSharedPtr< IDatasmithKeyValueProperty >& Property = MaterialElement->GetProperty(Index);
 			FName PropertyName(Property->GetName());
 
 			// Vector Params
diff --git a/Engine/Plugins/Experimental/Enterprise/DatasmithRuntime/Source/Private/MaterialImporter.cpp b/Engine/Plugins/Experimental/Enterprise/DatasmithRuntime/Source/Private/MaterialImporter.cpp
--- a/Engine/Plugins/Experimental/Enterprise/DatasmithRuntime/Source/Private/MaterialImporter.cpp
+++ b/Engine/Plugins/Experimental/Enterprise/DatasmithRuntime/Source/Private/MaterialImporter.cpp
@@ -68,6 +68,9 @@ namespace DatasmithRuntime
 				MaterialData.Object = TWeakObjectPtr<UObject>( UMaterialInstanceDynamic::Create( nullptr, nullptr) );
 #endif
 				check(MaterialData.Object.IsValid());
+
+				// Load metadata on newly created material asset if any
+				ApplyMetadata(MaterialData.MetadataId, MaterialData.GetObject());
 			}
 		}
 
@@ -152,6 +155,16 @@ namespace DatasmithRuntime
 			TSharedPtr< IDatasmithMasterMaterialElement > MaterialElement = StaticCastSharedPtr< IDatasmithMasterMaterialElement >( Element );
 
 			bCreationSuccessful = LoadMasterMaterial(MaterialInstance, MaterialElement);
+
+			// Add tracking on material's properties
+			if (bCreationSuccessful)
+			{
+				for (int Index = 0; Index < MaterialElement->GetPropertiesCount(); ++Index)
+				{
+					const TSharedPtr< IDatasmithKeyValueProperty >& Property = MaterialElement->GetProperty(Index);
+					DependencyList.Add(Property->GetNodeId(), { EDataType::Material, ElementId, 0xffff });
+				}
+			}
 		}
 		else if ( Element->IsA( EDatasmithElementType::UEPbrMaterial ) )
 		{
diff --git a/Engine/Plugins/Experimental/Enterprise/DatasmithRuntime/Source/Private/MeshImporter.cpp b/Engine/Plugins/Experimental/Enterprise/DatasmithRuntime/Source/Private/MeshImporter.cpp
--- a/Engine/Plugins/Experimental/Enterprise/DatasmithRuntime/Source/Private/MeshImporter.cpp
+++ b/Engine/Plugins/Experimental/Enterprise/DatasmithRuntime/Source/Private/MeshImporter.cpp
@@ -114,7 +114,10 @@ namespace DatasmithRuntime
 				};
 
 				AddToQueue(EQueueTask::MeshQueue, { TaskFunc, {EDataType::Mesh, MeshData.ElementId, 0 } });
-				TasksToComplete |=  EWorkerTask::MeshCreate;
+				TasksToComplete |= EWorkerTask::MeshCreate;
+
+				// If applicable, Apply metadata on newly created static mesh
+				ApplyMetadata(MeshData.MetadataId, StaticMesh);
 
 				MeshElementSet.Add(MeshData.ElementId);
 			}
diff --git a/Engine/Plugins/Experimental/Enterprise/DatasmithRuntime/Source/Private/SceneImporter.cpp b/Engine/Plugins/Experimental/Enterprise/DatasmithRuntime/Source/Private/SceneImporter.cpp
--- a/Engine/Plugins/Experimental/Enterprise/DatasmithRuntime/Source/Private/SceneImporter.cpp
+++ b/Engine/Plugins/Experimental/Enterprise/DatasmithRuntime/Source/Private/SceneImporter.cpp
@@ -6,22 +6,12 @@
 #include "LogCategory.h"
 #include "MaterialImportUtils.h"
 
+#include "DatasmithAssetUserData.h"
 #include "DatasmithNativeTranslator.h"
 #include "DatasmithMaterialElements.h"
 #include "IDatasmithSceneElements.h"
 
-#include "Camera/PlayerCameraManager.h"
-#include "CineCameraActor.h"
-#include "CineCameraComponent.h"
-#include "Components/SceneComponent.h"
-#include "Engine/Scene.h"
-#include "Engine/StaticMesh.h"
-#include "Engine/Texture2D.h"
-#include "Engine/World.h"
-#include "GameFramework/Actor.h"
-#include "GameFramework/Pawn.h"
-#include "GameFramework/PlayerController.h"
-#include "Kismet/GameplayStatics.h"
+#include "Async/Async.h"
 #include "Materials/MaterialInstanceDynamic.h"
 #include "ProfilingDebugging/MiscTrace.h"
 #include "RenderingThread.h"
@@ -29,19 +19,7 @@
 namespace DatasmithRuntime
 {
 	extern void UpdateMaterials(TSet<FSceneGraphId>& MaterialElementSet, TMap< FSceneGraphId, FAssetData >& AssetDataList);
-
-	void DeleteSceneComponent(USceneComponent* SceneComponent);
-
-	// Recursively destroy actor from world. Children are deleted first.
-	void DeleteActorRecursive(UWorld* GameWorld, AActor* Actor);
-
-	// Set scene component's visibility to false and remove geometry from static mesh component
-	// Call before deletion to ensure no asset is used
-	void HideSceneComponent(USceneComponent* SceneComponent);
-
-	// Tag added to created scene components and actors
-	// Used during reset process
-	const FName RuntimeTag("Datasmith.Runtime.Tag");
+	extern void HideSceneComponent(USceneComponent* SceneComponent);
 
 #ifdef LIVEUPDATE_TIME_LOGGING
 	Timer::Timer(double InTimeOrigin, const char* InText)
@@ -220,6 +198,25 @@ namespace DatasmithRuntime
 			);
 		}
 
+		if (ImportOptions.bImportMetaData)
+		{
+			// Start collection of metadata
+			MetadataCollect = Async(
+#if WITH_EDITOR
+				EAsyncExecution::LargeThreadPool,
+#else
+				EAsyncExecution::ThreadPool,
+#endif
+				[this]() -> void
+				{
+					for (int32 Index = 0; Index < this->SceneElement->GetMetaDataCount(); ++Index)
+					{
+						this->ProcessMetdata(this->SceneElement->GetMetaData(Index));
+					}
+				}
+			);
+		}
+
 		TasksToComplete |= EWorkerTask::SetupTasks;
 	}
 
@@ -254,50 +251,6 @@ namespace DatasmithRuntime
 		OnGoingTasks.Reserve(TextureElementSet.Num() + MeshElementSet.Num());
 	}
 
-	void FSceneImporter::ProcessActorElement(const TSharedPtr< IDatasmithActorElement >& ActorElement, FSceneGraphId ParentId)
-	{
-		TRACE_CPUPROFILER_EVENT_SCOPE(FSceneImporter::ProcessActorElement);
-
-		FActorData& ActorData = FindOrAddActorData(ActorElement);
-
-		if (ActorData.HasState(EAssetState::Processed))
-		{
-			return;
-		}
-
-		ActorData.ParentId = ParentId;
-		ActorData.WorldTransform = FTransform( ActorElement->GetRotation(), ActorElement->GetTranslation(), ActorElement->GetScale() );
-
-		bool bProcessSuccessful = false;
-
-		if (ActorElement->IsA(EDatasmithElementType::StaticMeshActor))
-		{
-			IDatasmithMeshActorElement* MeshActorElement = static_cast<IDatasmithMeshActorElement*>(ActorElement.Get());
-			
-			ActorData.Type = EDataType::MeshActor;
-			bProcessSuccessful = ProcessMeshActorData(ActorData, MeshActorElement);
-		}
-		else if (ActorElement->IsA(EDatasmithElementType::Light))
-		{
-			IDatasmithLightActorElement* LightActorElement = static_cast<IDatasmithLightActorElement*>(ActorElement.Get());
-
-			ActorData.Type = EDataType::LightActor;
-			bProcessSuccessful = ProcessLightActorData(ActorData, LightActorElement);
-		}
-		else if (ActorElement->IsA(EDatasmithElementType::Camera))
-		{
-			IDatasmithCameraActorElement* CameraElement = static_cast<IDatasmithCameraActorElement*>(ActorElement.Get());
-
-			bProcessSuccessful = ProcessCameraActorData(ActorData, CameraElement);
-		}
-
-		if (!bProcessSuccessful)
-		{
-			CreateActorComponent(ActorData, ActorElement);
-			ActorData.AddState(EAssetState::Processed | EAssetState::Completed);
-		}
-	}
-
 	void FSceneImporter::Tick(float DeltaSeconds)
 	{
 		TRACE_CPUPROFILER_EVENT_SCOPE(FSceneImporter::Tick);
@@ -408,6 +361,12 @@ namespace DatasmithRuntime
 			TasksToComplete &= ~EWorkerTask::SetupTasks;
 		}
 
+		// Do not proceed further if metadata collection is not complete
+		if (MetadataCollect.IsValid() && !MetadataCollect.IsReady())
+		{
+			return;
+		}
+
 		bContinue = FPlatformTime::Seconds() < EndTime;
 
 		if (bContinue && EnumHasAnyFlags(TasksToComplete, EWorkerTask::MeshCreate))
@@ -565,16 +524,6 @@ namespace DatasmithRuntime
 
 			UE_LOG(LogDatasmithRuntime, Log, TEXT("Total load time is [%d min %.3f s]"), ElapsedMin, ElapsedSeconds);
 #endif
-
-			// Return if async tasks are not completed
-			for (TFuture<bool>& OnGoingTask : OnGoingTasks)
-			{
-				if (!OnGoingTask.IsReady() && TasksToComplete != EWorkerTask::NoTask)
-				{
-					ensure(false);
-					break;
-				}
-			}
 		}
 	}
 
@@ -597,6 +546,7 @@ namespace DatasmithRuntime
 		// Clear all cached data if it is a new scene
 		SceneElement.Reset();
 		LastSceneGuid = FGuid();
+		MetadataCollect = TFuture<void>();
 
 		TasksToComplete = EWorkerTask::ResetScene;
 
@@ -876,6 +826,12 @@ namespace DatasmithRuntime
 					AddToQueue(EQueueTask::DeleteCompQueue, { TaskFunc, FReferencer(ElementId) } );
 					TasksToComplete |= EWorkerTask::DeleteComponent;
 				}
+				// Remove metadata from list of tracked elements
+				else if (Elements[ElementId]->IsA(EDatasmithElementType::MetaData))
+				{
+					AddToQueue(EQueueTask::DeleteCompQueue, { TaskFunc, FReferencer(ElementId) } );
+					TasksToComplete |= EWorkerTask::DeleteComponent;
+				}
 				else
 				{
 					UE_LOG(LogDatasmithRuntime, Error, TEXT("Element %d (%s) was not found"), ElementId, Elements[ElementId]->GetName());
@@ -904,6 +860,64 @@ namespace DatasmithRuntime
 				AddToQueue(EQueueTask::NonAsyncQueue, { AssignMaterialFunc, *MaterialElementIdPtr, MoveTemp(Referencer) });
 				TasksToComplete |= EWorkerTask::MaterialAssign;
 			}
+			// The value of a property has changed, process it
+			else if (Element->IsA(EDatasmithElementType::KeyValueProperty))
+			{
+				// If it is a material's property, invalidate material and queue its processing
+				if (Referencer.GetType() == EDataType::Material)
+				{
+					if (AssetDataList.Contains(Referencer.ElementId))
+					{
+						FActionTaskFunction TaskFunc = [this, ElementId = Referencer.ElementId](UObject*, const FReferencer&) -> EActionResult::Type
+						{
+							FAssetData& MaterialData = this->AssetDataList[ElementId];
+
+							MaterialData.SetState(EAssetState::Unknown);
+
+							this->ProcessMaterialData(MaterialData);
+
+							ActionCounter.Increment();
+
+							return EActionResult::Succeeded;
+						};
+
+						AddToQueue(EQueueTask::UpdateQueue, { MoveTemp(TaskFunc), FReferencer() } );
+					}
+				}
+				// If it is a metadata's property, queue its application to the associated element
+				else if (Referencer.GetType() == EDataType::Metadata)
+				{
+					//
+					if (Elements.Contains(Referencer.ElementId))
+					{
+						FActionTaskFunction TaskFunc = [this, ElementId = Referencer.ElementId](UObject*, const FReferencer&) -> EActionResult::Type
+						{
+							TSharedPtr<IDatasmithMetaDataElement> MetadataElement = StaticCastSharedPtr<IDatasmithMetaDataElement>(Elements[ElementId]);
+
+							const TSharedPtr< IDatasmithElement >& AssociatedElement = MetadataElement->GetAssociatedElement();
+							if (AssociatedElement && Elements.Contains(AssociatedElement->GetNodeId()))
+							{
+								const FSceneGraphId AssociatedId = AssociatedElement->GetNodeId();
+
+								if (AssetDataList.Contains(AssociatedId))
+								{
+									ApplyMetadata(ElementId, AssetDataList[AssociatedId].GetObject());
+								}
+								else if (ActorDataList.Contains(AssociatedId))
+								{
+									ApplyMetadata(ElementId, ActorDataList[AssociatedId].GetObject());
+								}
+							}
+
+							ActionCounter.Increment();
+
+							return EActionResult::Succeeded;
+						};
+
+						AddToQueue(EQueueTask::UpdateQueue, { MoveTemp(TaskFunc), FReferencer() } );
+					}
+				}
+			}
 		}
 	}
 
@@ -1021,6 +1035,10 @@ namespace DatasmithRuntime
 					}
 				}
 			}
+			else if (ImportOptions.bImportMetaData && ElementPtr->IsA(EDatasmithElementType::MetaData))
+			{
+				ProcessMetdata(StaticCastSharedPtr<IDatasmithMetaDataElement>(ElementPtr));
+			}
 			else if (ElementPtr->IsA(EDatasmithElementType::Actor))
 			{
 				UE_LOG(LogDatasmithRuntime, Log, TEXT("IncrementalAdditions: Actor %s (%d)"), ElementPtr->GetName(), ElementPtr->GetNodeId());
@@ -1077,6 +1095,16 @@ namespace DatasmithRuntime
 			);
 		}
 
+		for (int32 Index = 0; Index < SceneElement->GetMetaDataCount(); ++Index)
+		{
+			const FSceneGraphId ElementId = SceneElement->GetMetaData(Index)->GetNodeId();
+
+			if (this->Elements.Contains(ElementId))
+			{
+				this->Elements[ElementId] = SceneElement->GetMetaData(Index);
+			}
+		}
+
 		// Clear 'Processed' state of modified elements
 		for (TSharedPtr<IDatasmithElement>& ElementPtr : UpdateContext.Updates)
 		{
@@ -1125,76 +1153,6 @@ namespace DatasmithRuntime
 		}
 	}
 
-	bool FSceneImporter::DeleteData()
-	{
-		// Terminate all rendering commands before deleting any component and asset
-		FlushRenderingCommands();
-
-		bool bGarbageCollect = false;
-
-		TArray<AActor*> ChildActors;
-		RootComponent->GetOwner()->GetAttachedActors(ChildActors);
-
-		// Check to see if last import required to build the actor's hierarchy
-		bool bHasHierarchy = false;
-		for (AActor* ChildActor : ChildActors)
-		{
-			if (ChildActor->ActorHasTag(RuntimeTag))
-			{
-				bHasHierarchy = true;
-				break;
-			}
-		}
-
-		// If last import created actors, delete them from bottom to top
-		if (bHasHierarchy)
-		{
-			// Remove hold onto scene component as they will be deleted along with the owning actor.
-			for (TPair< FSceneGraphId, FActorData >& Pair : ActorDataList)
-			{
-				Pair.Value.Object.Reset();
-			}
-
-			UWorld* GameWorld = RootComponent->GetOwner()->GetWorld();
-
-			for (AActor* ChildActor : ChildActors)
-			{
-				// Only delete actors created by runtime import
-				if (ChildActor->ActorHasTag(RuntimeTag))
-				{
-					ChildActor->DetachFromActor(FDetachmentTransformRules::KeepRelativeTransform);
-					DeleteActorRecursive(GameWorld, ChildActor);
-				}
-			}
-
-			TArray<USceneComponent*> SceneComponents = RootComponent->GetAttachChildren();
-			for (USceneComponent* SceneComponent : SceneComponents)
-			{
-				// Only delete components created by runtime import
-				if (SceneComponent->ComponentHasTag(RuntimeTag))
-				{
-					DeleteSceneComponent(SceneComponent);
-				}
-			}
-
-			bGarbageCollect = true;
-		}
-		else
-		{
-			for (TPair< FSceneGraphId, FActorData >& Pair : ActorDataList)
-			{
-				bGarbageCollect |= DeleteComponent(Pair.Value);
-			}
-		}
-
-		for (TPair< FSceneGraphId, FAssetData >& Entry : AssetDataList)
-		{
-			bGarbageCollect |= DeleteAsset(Entry.Value);
-		}
-
-		return bGarbageCollect;
-	}
-
 	EActionResult::Type FSceneImporter::DeleteElement(FSceneGraphId ElementId)
 	{
 		TSharedPtr<IDatasmithElement> ElementPtr;
@@ -1245,66 +1203,18 @@ namespace DatasmithRuntime
 			return DeleteAsset(AssetData) ? EActionResult::Succeeded : EActionResult::Failed;
 		}
 
-		ensure(ActorDataList.Contains(ElementId));
-
-		FActorData ActorData(DirectLink::InvalidId);
-		if (!ActorDataList.RemoveAndCopyValue(ElementId, ActorData))
-		{
-			return EActionResult::Failed;
-		}
-
-		return DeleteComponent(ActorData) ? EActionResult::Succeeded : EActionResult::Failed;
-	}
-
-	bool FSceneImporter::DeleteComponent(FActorData& ActorData)
-	{
-		if (USceneComponent* SceneComponent = ActorData.GetObject<USceneComponent>())
+		if (ActorDataList.Contains(ElementId))
 		{
-			if (ImportOptions.BuildHierarchy !=  EBuildHierarchyMethod::None)
+			FActorData ActorData(DirectLink::InvalidId);
+			if (!ActorDataList.RemoveAndCopyValue(ElementId, ActorData))
 			{
-				USceneComponent* ParentComponent = SceneComponent->GetAttachmentRoot();
-
-				AActor* Actor = SceneComponent->GetOwner();
-
-				if (ParentComponent->GetOwner() == Actor)
-				{
-					TArray<USceneComponent*> Children = SceneComponent->GetAttachChildren();
-
-					for (USceneComponent* ChildComponent : Children)
-					{
-						ChildComponent->DetachFromComponent(FDetachmentTransformRules::KeepRelativeTransform);
-						ChildComponent->AttachToComponent(ParentComponent, FAttachmentTransformRules::KeepRelativeTransform);
-					}
-
-					DeleteSceneComponent(SceneComponent);
-				}
-				else
-				{
-					AActor* ParentActor = ParentComponent->GetOwner();
-
-					TArray<AActor*> ChildActors;
-					Actor->GetAttachedActors(ChildActors);
-
-					for (AActor* ChildActor : ChildActors)
-					{
-						ChildActor->DetachFromActor(FDetachmentTransformRules::KeepRelativeTransform);
-						ChildActor->AttachToActor(ParentActor, FAttachmentTransformRules::KeepRelativeTransform);
-					}
-
-					RootComponent->GetOwner()->GetWorld()->EditorDestroyActor(Actor, true);
-				}
-			}
-			else if (SceneComponent->GetAttachmentRoot() == RootComponent.Get())
-			{
-				DeleteSceneComponent(SceneComponent);
+				return EActionResult::Failed;
 			}
 
-			ActorData.Object.Reset();
-
-			return true;
+			return DeleteComponent(ActorData) ? EActionResult::Succeeded : EActionResult::Failed;
 		}
 
-		return false;
+		return EActionResult::Succeeded;
 	}
 
 	bool FSceneImporter::DeleteAsset(FAssetData& AssetData)
@@ -1320,402 +1230,77 @@ namespace DatasmithRuntime
 		return false;
 	}
 
-	bool FSceneImporter::ProcessCameraActorData(FActorData& ActorData, IDatasmithCameraActorElement* CameraElement)
+	void FSceneImporter::ProcessMetdata(const TSharedPtr<IDatasmithMetaDataElement>& MetadataElement)
 	{
-		if (ActorData.HasState(EAssetState::Processed) || ActorData.Object.IsValid())
+		// Process metadata only if it has properties and the associated element is tracked
+		if (MetadataElement && MetadataElement->GetPropertiesCount() > 0)
 		{
-			ActorData.AddState(EAssetState::Processed);
-			return true;
-		}
-		
-		// Check to see if the camera must be updated or not
-		// Update if only the current actor is the only one with a valid source and the source has changed
-		bool bModifyViewpoint = ImportOptions.bModifyViewpoint;
-
-		TArray<AActor*> Actors;
-		UGameplayStatics::GetAllActorsOfClass(RootComponent->GetOwner()->GetWorld(), ADatasmithRuntimeActor::StaticClass(), Actors);
+			const TSharedPtr< IDatasmithElement >& AssociatedElement = MetadataElement->GetAssociatedElement();
+			const FSceneGraphId AssociatedId = AssociatedElement ? AssociatedElement->GetNodeId() : DirectLink::InvalidId;
 
-		if (Actors.Num() > 1)
-		{
-			for (AActor* Actor : Actors)
+			if (Elements.Contains(AssociatedId))
 			{
-				if (Actor == RootComponent->GetOwner())
+				if (AssetDataList.Contains(AssociatedId))
 				{
-					continue;
+					AssetDataList[AssociatedId].MetadataId = MetadataElement->GetNodeId();
 				}
-
-				if (ADatasmithRuntimeActor* RuntimeActor = Cast<ADatasmithRuntimeActor>(Actor))
+				else if (ActorDataList.Contains(AssociatedId))
 				{
-					bModifyViewpoint &= RuntimeActor->GetSourceName() == TEXT("None");
+					ActorDataList[AssociatedId].MetadataId = MetadataElement->GetNodeId();
 				}
-			}
-		}
 
-		bModifyViewpoint &= LastSceneGuid != SceneElement->GetSharedState()->GetGuid();
-
-		if (bModifyViewpoint)
-		{
-			if (APlayerController* PlayerController = UGameplayStatics::GetPlayerController(RootComponent->GetOwner()->GetWorld(), 0))
-			{
-				PlayerController->SetControlRotation(ActorData.WorldTransform.GetRotation().Rotator());
-				if (APawn* Pawn = PlayerController->GetPawn())
-				{
-					Pawn->SetActorLocationAndRotation(ActorData.WorldTransform.GetLocation(), ActorData.WorldTransform.GetRotation(), false);
-				}
+				Elements.Add(MetadataElement->GetNodeId(), MetadataElement);
 			}
-		}
-
-		if (ImportOptions.BuildHierarchy != EBuildHierarchyMethod::None)
-		{
-			UCineCameraComponent* CameraComponent = ActorData.GetObject<UCineCameraComponent>();
-
-			if (CameraComponent == nullptr)
+			else if (AssociatedElement == SceneElement)
 			{
-				if (ImportOptions.BuildHierarchy != EBuildHierarchyMethod::None)
-				{
-					UWorld* World = RootComponent->GetOwner()->GetWorld();
-
-					ACineCameraActor* CameraActor = Cast< ACineCameraActor >( World->SpawnActor( ACineCameraActor::StaticClass(), nullptr, nullptr ) );
-					CameraActor->Rename(CameraElement->GetName(), nullptr, REN_NonTransactional | REN_DontCreateRedirectors);
-#if WITH_EDITOR
-					CameraActor->SetActorLabel(CameraElement->GetLabel());
-#endif
-					CameraComponent = CameraActor->GetCineCameraComponent();
-				}
-				else
-				{
-					FName ComponentName = NAME_None;
-					ComponentName = MakeUniqueObjectName(RootComponent->GetOwner(), UCineCameraComponent::StaticClass(), CameraElement->GetLabel());
-					CameraComponent = NewObject< UCineCameraComponent >(RootComponent->GetOwner(), ComponentName);
-				}
-
-				ActorData.Object = CameraComponent;
+				Elements.Add(MetadataElement->GetNodeId(), MetadataElement);
+				ApplyMetadata(MetadataElement->GetNodeId(), RootComponent.Get());
 			}
-			//else
-			//{
-			//	CameraComponent->GetOwner()->UpdateComponentTransforms();
-			//	CameraComponent->GetOwner()->MarkComponentsRenderStateDirty();
-			//}
-
-			CameraComponent->Filmback.SensorWidth = CameraElement->GetSensorWidth();
-			CameraComponent->Filmback.SensorHeight = CameraElement->GetSensorWidth() / CameraElement->GetSensorAspectRatio();
-			CameraComponent->LensSettings.MaxFStop = 32.0f;
-			CameraComponent->CurrentFocalLength = CameraElement->GetFocalLength();
-			CameraComponent->CurrentAperture = CameraElement->GetFStop();
-
-			CameraComponent->FocusSettings.FocusMethod = CameraElement->GetEnableDepthOfField() ? ECameraFocusMethod::Manual : ECameraFocusMethod::DoNotOverride;
-			CameraComponent->FocusSettings.ManualFocusDistance = CameraElement->GetFocusDistance();
-
-			if (const IDatasmithPostProcessElement* PostProcess = CameraElement->GetPostProcess().Get())
-			{
-				FPostProcessSettings& PostProcessSettings = CameraComponent->PostProcessSettings;
-
-				if ( !FMath::IsNearlyEqual( PostProcess->GetTemperature(), 6500.f ) )
-				{
-					PostProcessSettings.bOverride_WhiteTemp = true;
-					PostProcessSettings.WhiteTemp = PostProcess->GetTemperature();
-				}
-
-				if ( PostProcess->GetVignette() > 0.f )
-				{
-					PostProcessSettings.bOverride_VignetteIntensity = true;
-					PostProcessSettings.VignetteIntensity = PostProcess->GetVignette();
-				}
-
-				if (PostProcess->GetColorFilter() != FLinearColor::Black && PostProcess->GetColorFilter() != FLinearColor::White )
-				{
-					PostProcessSettings.bOverride_FilmWhitePoint = true;
-					PostProcessSettings.FilmWhitePoint = PostProcess->GetColorFilter();
-				}
-
-				if ( !FMath::IsNearlyEqual( PostProcess->GetSaturation(), 1.f ) )
-				{
-					PostProcessSettings.bOverride_ColorSaturation = true;
-					PostProcessSettings.ColorSaturation.W = PostProcess->GetSaturation();
-				}
-
-				if ( PostProcess->GetCameraISO() > 0.f || PostProcess->GetCameraShutterSpeed() > 0.f || PostProcess->GetDepthOfFieldFstop() > 0.f )
-				{
-					PostProcessSettings.bOverride_AutoExposureMethod = true;
-					PostProcessSettings.AutoExposureMethod = EAutoExposureMethod::AEM_Manual;
-
-					if ( PostProcess->GetCameraISO() > 0.f )
-					{
-						PostProcessSettings.bOverride_CameraISO = true;
-						PostProcessSettings.CameraISO = PostProcess->GetCameraISO();
-					}
-
-					if ( PostProcess->GetCameraShutterSpeed() > 0.f )
-					{
-						PostProcessSettings.bOverride_CameraShutterSpeed = true;
-						PostProcessSettings.CameraShutterSpeed = PostProcess->GetCameraShutterSpeed();
-					}
-
-					if ( PostProcess->GetDepthOfFieldFstop() > 0.f )
-					{
-						PostProcessSettings.bOverride_DepthOfFieldFstop = true;
-						PostProcessSettings.DepthOfFieldFstop = PostProcess->GetDepthOfFieldFstop();
-					}
-				}
-			}
-		}
-
-		FinalizeComponent(ActorData);
-
-		ActorData.SetState(EAssetState::Processed | EAssetState::Completed);
-
-		return true;
-	}
-
-	FActorData& FSceneImporter::FindOrAddActorData(const TSharedPtr<IDatasmithActorElement>& ActorElement)
-	{
-		FSceneGraphId ElementId = ActorElement->GetNodeId();
-
-		if (!Elements.Contains(ElementId))
-		{
-			Elements.Add(ElementId, ActorElement);
-
-			FActorData ActorData(ElementId);
-			return ActorDataList.Emplace(ElementId, MoveTemp(ActorData));
 		}
-
-		ensure(ActorDataList.Contains(ElementId));
-		return ActorDataList[ElementId];
 	}
 
-	void FSceneImporter::CreateActorComponent(FActorData& ActorData, const TSharedPtr< IDatasmithActorElement >& ActorElement)
+	// Logic borrowed from FDatasmithImporter::ImportMetaDataForObject
+	void FSceneImporter::ApplyMetadata(FSceneGraphId MetadataId, UObject* Object)
 	{
-		const bool bBuildHierarchy = ImportOptions.BuildHierarchy == EBuildHierarchyMethod::Unfiltered || (ImportOptions.BuildHierarchy == EBuildHierarchyMethod::Simplified && ActorElement->GetChildrenCount() > 0);
-		if (!bBuildHierarchy)
+		if ( !Object || !Object->GetClass()->ImplementsInterface( UInterface_AssetUserData::StaticClass() ) || MetadataId == DirectLink::InvalidId)
 		{
-			ActorData.SetState(EAssetState::Skipped);
 			return;
 		}
 
-		bool bMustProcess = true;
-
-		if (ImportOptions.BuildHierarchy == EBuildHierarchyMethod::Simplified)
-		{
-			TFunction<bool(const TSharedPtr< IDatasmithActorElement >&)> MustKeep;
-			MustKeep = [this, &MustKeep](const TSharedPtr< IDatasmithActorElement >& ActorElementToCheck) -> bool
-			{
-				for (int32 Index = 0; Index < ActorElementToCheck->GetChildrenCount(); ++Index)
-				{
-					if (MustKeep(ActorElementToCheck->GetChild(Index)))
-					{
-						return true;
-					}
-				}
-
-				const bool bMustKeep = ActorElementToCheck->IsA(EDatasmithElementType::Camera)
-					|| ActorElementToCheck->IsA(EDatasmithElementType::Light)
-					|| ActorElementToCheck->IsA(EDatasmithElementType::StaticMeshActor);
-
-				if (!bMustKeep)
-				{
-					FActorData& LocalActorData = this->FindOrAddActorData(ActorElementToCheck);
-					LocalActorData.SetState(EAssetState::Processed | EAssetState::Completed | EAssetState::Skipped);
-				}
-
-				return bMustKeep;
-			};
-
-			bMustProcess = MustKeep(ActorElement);
-
-			if (bMustProcess && ActorElement->GetChildrenCount() == 1 && ActorElement->GetChild(0)->GetChildrenCount() == 0)
-			{
-				ActorElement->GetChild(0)->SetIsAComponent(false);
-				ActorData.SetState(EAssetState::Skipped);
-				bMustProcess = false;
-			}
-		}
-
-		if(bMustProcess)
+		if (IDatasmithMetaDataElement* MetadataElement = static_cast<IDatasmithMetaDataElement*>(Elements[MetadataId].Get()))
 		{
-			USceneComponent* SceneComponent = ActorData.GetObject<USceneComponent>();
-
-			if (SceneComponent == nullptr)
+			if (IInterface_AssetUserData* AssetUserData = Cast< IInterface_AssetUserData >(Object))
 			{
-				AActor* Actor = RootComponent->GetOwner()->GetWorld()->SpawnActor(AActor::StaticClass(), nullptr, nullptr);
-				Actor->Rename(ActorElement->GetName(), nullptr, REN_NonTransactional | REN_DontCreateRedirectors);
-#if WITH_EDITOR
-				Actor->SetActorLabel( ActorElement->GetLabel() );
-#endif
+				UDatasmithAssetUserData* DatasmithUserData = AssetUserData->GetAssetUserData< UDatasmithAssetUserData >();
 
-				Actor->Tags.Empty(ActorElement->GetTagsCount());
-				for (int32 Index = 0; Index < ActorElement->GetTagsCount(); ++Index)
+				if ( !DatasmithUserData )
 				{
-					Actor->Tags.Add(ActorElement->GetTag(Index));
+					DatasmithUserData = NewObject<UDatasmithAssetUserData>( Object, NAME_None, RF_Public | RF_Transactional );
+					AssetUserData->AddAssetUserData( DatasmithUserData );
 				}
 
-				// #ue_datasmithruntime: Add Layers and metadata
-
-				SceneComponent = Actor->GetRootComponent();
+				UDatasmithAssetUserData::FMetaDataContainer MetaData;
 
-				if ( !SceneComponent )
-				{
-					SceneComponent = NewObject< USceneComponent >( Actor, USceneComponent::StaticClass(), ActorElement->GetLabel(), RF_Transactional );
-
-					Actor->AddInstanceComponent( SceneComponent );
-					Actor->SetRootComponent( SceneComponent );
-				}
-
-				ActorData.Object = TWeakObjectPtr<UObject>(SceneComponent);
-			}
+				const int32 PropertiesCount = MetadataElement->GetPropertiesCount();
+				MetaData.Reserve( PropertiesCount + 1 );
 
-			FinalizeComponent(ActorData);
-		}
-	}
-
-	void FSceneImporter::FinalizeComponent(FActorData& ActorData)
-	{
-		USceneComponent* SceneComponent = ActorData.GetObject<USceneComponent>();
-		if (!SceneComponent)
-		{
-			return;
-		}
-
-		if (SceneComponent->IsRegistered())
-		{
-			SceneComponent->UnregisterComponent();
-		}
+				// Add associated element's unique id
+				MetaData.Add( UDatasmithAssetUserData::UniqueIdMetaDataKey, MetadataElement->GetAssociatedElement()->GetName() );
 
-		// Find the right parent
-		FSceneGraphId ParentId = ActorData.ParentId;
-		while (ParentId != DirectLink::InvalidId && ActorDataList[ParentId].HasState(EAssetState::Skipped))
-		{
-			ParentId = ActorDataList[ParentId].ParentId;
-		}
-
-		if (ImportOptions.BuildHierarchy != EBuildHierarchyMethod::None && ParentId != DirectLink::InvalidId)
-		{
-			// If hierarchy building is required, wait for parent component to be created if applicable
-			if (ActorDataList[ParentId].GetObject<USceneComponent>() == nullptr)
-			{
-				FActionTaskFunction FinalizeComponentFunc = [this,ParentId](UObject* Object, const FReferencer& Referencer) -> EActionResult::Type
+				// Add Datasmith metadata's properties
+				for ( int32 PropertyIndex = 0; PropertyIndex < PropertiesCount; ++PropertyIndex )
 				{
-					FActorData& ActorData = ActorDataList[Referencer.GetId()];
-
-					if (ActorDataList[ParentId].GetObject<USceneComponent>() == nullptr)
-					{
-						return EActionResult::Retry;
-					}
-
-					this->FinalizeComponent(ActorData);
-
-					return EActionResult::Succeeded;
-				};
-
-				AddToQueue(EQueueTask::NonAsyncQueue, { FinalizeComponentFunc, { EDataType::Actor, ActorData.ElementId, 0 } });
-				TasksToComplete |= EWorkerTask::ComponentFinalize;
-
-				return;
-			}
-		}
-
-		const bool bUseParent = ParentId != DirectLink::InvalidId && ImportOptions.BuildHierarchy != EBuildHierarchyMethod::None;
-		USceneComponent* ParentComponent = bUseParent ? ActorDataList[ParentId].GetObject<USceneComponent>() : RootComponent.Get();
-
-		if (ParentComponent)
-		{
-			const IDatasmithActorElement* ActorElement = static_cast<IDatasmithActorElement*>(Elements[ActorData.ElementId].Get());
-
-			const FTransform ParentToWorld = ParentComponent->GetComponentToWorld();
-			const FTransform RelativeTransform = ActorData.WorldTransform.GetRelativeTransform(ParentToWorld);
+					const TSharedPtr<IDatasmithKeyValueProperty>& Property = MetadataElement->GetProperty( PropertyIndex );
+					MetaData.Add( Property->GetName(), Property->GetValue() );
 
-			SceneComponent->SetRelativeTransform( RelativeTransform );
-			SceneComponent->SetVisibility(ActorElement->GetVisibility());
-			SceneComponent->SetMobility( EComponentMobility::Movable );
-
-			if (SceneComponent->GetAttachParent() != ParentComponent)
-			{
-				SceneComponent->DetachFromComponent(FDetachmentTransformRules::KeepRelativeTransform);
-				SceneComponent->AttachToComponent( ParentComponent, FAttachmentTransformRules::KeepRelativeTransform );
-			}
-
-			SceneComponent->ComponentTags.Empty(ActorElement->GetTagsCount());
-			for (int32 Index = 0; Index < ActorElement->GetTagsCount(); ++Index)
-			{
-				SceneComponent->ComponentTags.Add(ActorElement->GetTag(Index));
-			}
-
-			// Add runtime tag if scene component is attached to runtime actor's component
-			if (ParentComponent == RootComponent)
-			{
-				SceneComponent->ComponentTags.Add(RuntimeTag);
-			}
-
-			if (ImportOptions.BuildHierarchy != EBuildHierarchyMethod::None)
-			{
-				if (ActorElement->IsAComponent() && ParentComponent && SceneComponent->GetOwner() != ParentComponent->GetOwner())
-				{
-					FName UniqueName = MakeUniqueObjectName(ParentComponent->GetOwner(), SceneComponent->GetClass(), ActorElement->GetName());
-					SceneComponent->Rename(*UniqueName.ToString(), ParentComponent->GetOwner(), REN_NonTransactional | REN_DontCreateRedirectors);
-				}
-
-				SceneComponent->GetOwner()->AddInstanceComponent(SceneComponent);
-
-				// Add runtime tag on actor created on import.
-				if (SceneComponent->GetOwner() != RootComponent->GetOwner() && !SceneComponent->GetOwner()->ActorHasTag(RuntimeTag))
-				{
-					SceneComponent->GetOwner()->Tags.Add(RuntimeTag);
+					DependencyList.Add(Property->GetNodeId(), { EDataType::Metadata, MetadataId, 0xffff });
 				}
-			}
-
-			SceneComponent->RegisterComponentWithWorld(RootComponent->GetOwner()->GetWorld());
-		}
-		else
-		{
-			ensure(false);
-		}
-	}
-
-	void DeleteActorRecursive(UWorld* GameWorld, AActor* Actor)
-	{
-		TArray<AActor*> ChildActors;
-		Actor->GetAttachedActors(ChildActors);
-
-		for (AActor* ChildActor : ChildActors)
-		{
-			ChildActor->DetachFromActor(FDetachmentTransformRules::KeepRelativeTransform);
-			DeleteActorRecursive(GameWorld, ChildActor);
-		}
-
-		ensure(GameWorld->DestroyActor(Actor));
-
-		// Since deletion can be delayed, rename to avoid future name collision
-		// Call UObject::Rename directly on actor to avoid AActor::Rename which unnecessarily sunregister and re-register components
-		Actor->UObject::Rename( nullptr, GetTransientPackage(), REN_DontCreateRedirectors | REN_ForceNoResetLoaders );
-	}
 
-	void HideSceneComponent(USceneComponent* SceneComponent)
-	{
-		if (SceneComponent)
-		{
-			SceneComponent->SetVisibility(false);
+				MetaData.KeySort(FNameLexicalLess());
 
-			if (UStaticMeshComponent* MeshComponent = Cast< UStaticMeshComponent >(SceneComponent))
-			{
-				MeshComponent->OverrideMaterials.Reset();
-				MeshComponent->SetStaticMesh(nullptr);
-				MeshComponent->ReleaseResources();
+				DatasmithUserData->MetaData = MoveTemp( MetaData );
 			}
-
-			SceneComponent->MarkRenderStateDirty();
-			SceneComponent->MarkRenderDynamicDataDirty();
 		}
 	}
 
-	void DeleteSceneComponent(USceneComponent* SceneComponent)
-	{
-		SceneComponent->UnregisterComponent();
-
-		SceneComponent->DetachFromComponent(FDetachmentTransformRules::KeepRelativeTransform);
-
-		SceneComponent->ClearFlags(RF_AllFlags);
-		SceneComponent->SetFlags(RF_Transient);
-		SceneComponent->Rename(nullptr, nullptr, REN_NonTransactional | REN_DontCreateRedirectors);
-		SceneComponent->MarkPendingKill();
-	}
-
 } // End of namespace DatasmithRuntime
\ No newline at end of file
diff --git a/Engine/Plugins/Experimental/Enterprise/DatasmithRuntime/Source/Private/SceneImporter.h b/Engine/Plugins/Experimental/Enterprise/DatasmithRuntime/Source/Private/SceneImporter.h
--- a/Engine/Plugins/Experimental/Enterprise/DatasmithRuntime/Source/Private/SceneImporter.h
+++ b/Engine/Plugins/Experimental/Enterprise/DatasmithRuntime/Source/Private/SceneImporter.h
@@ -26,6 +26,7 @@ class IDatasmithActorElement;
 class IDatasmithElement;
 class IDatasmithCameraActorElement;
 class IDatasmithLightActorElement;
+class IDatasmithMetaDataElement;
 class IDatasmithMeshActorElement;
 class IDatasmithScene;
 class IDatasmithTranslator;
@@ -134,6 +135,7 @@ namespace DatasmithRuntime
 		Actor       = 5,
 		MeshActor   = 6,
 		LightActor  = 7,
+		Metadata    = 8,
 	};
 
 	/**
@@ -173,6 +175,7 @@ namespace DatasmithRuntime
 		{
 		}
 
+		EDataType GetType() const { return EDataType(Type); }
 		FSceneGraphId GetId() const { return (FSceneGraphId)ElementId; }
 	};
 
@@ -186,6 +189,9 @@ namespace DatasmithRuntime
 		/** Identifier of the associated Datasmith element */
 		FSceneGraphId ElementId;
 
+		/** Identifier of the associated metadata element */
+		FSceneGraphId MetadataId;
+
 		EDataType Type;
 
 		/** UObject associated with the element */
@@ -199,6 +205,7 @@ namespace DatasmithRuntime
 
 		FBaseData(FSceneGraphId InElementId, EDataType InType = EDataType::None)
 			: ElementId(InElementId)
+			, MetadataId(DirectLink::InvalidId)
 			, Type(InType)
 		{
 			DataState.store(EAssetState::Unknown);
@@ -207,6 +214,7 @@ namespace DatasmithRuntime
 		FBaseData(const FBaseData& Other)
 		{
 			ElementId = Other.ElementId;
+			MetadataId = Other.MetadataId;
 			Type = Other.Type;
 			Object = Other.Object;
 			DataState.store(Other.DataState.load());
@@ -638,6 +646,12 @@ namespace DatasmithRuntime
 
 		FActorData& FindOrAddActorData(const TSharedPtr< IDatasmithActorElement >& ActorElement);
 
+		/** Add metadata to the tracking system if it applies to any tracked asset or actor or the scene */
+		void ProcessMetdata(const TSharedPtr< IDatasmithMetaDataElement >& MetadataElement);
+
+		/** Apply metadata to input object if applicable */
+		void ApplyMetadata(FSceneGraphId MetadataId, UObject* Object);
+
 	private:
 		/** DatasmithRuntime actor associated with this importer */
 		TWeakObjectPtr<USceneComponent> RootComponent;
@@ -680,6 +694,7 @@ namespace DatasmithRuntime
 		/** Array of queues dequeued during the import process */
 		TQueue< FActionTask, EQueueMode::Mpsc > ActionQueues[EQueueTask::MaxQueues];
 
+		/** Array of asynchronous tasks started during the import */
 		TArray<TFuture<bool>> OnGoingTasks;
 
 		/** Flag used to properly sequence the import process */
@@ -691,6 +706,9 @@ namespace DatasmithRuntime
 		/** Specifies options to use during the import */
 		FDatasmithRuntimeImportOptions ImportOptions;
 
+		/** Future for asynchronous task to collect metadata */
+		TFuture<void> MetadataCollect;
+
 		/** Miscellaneous counters used to report progress */
 		float& OverallProgress;
 		FThreadSafeCounter ActionCounter;
diff --git a/Engine/Plugins/Experimental/Enterprise/DatasmithRuntime/Source/Private/TextureImporter.cpp b/Engine/Plugins/Experimental/Enterprise/DatasmithRuntime/Source/Private/TextureImporter.cpp
--- a/Engine/Plugins/Experimental/Enterprise/DatasmithRuntime/Source/Private/TextureImporter.cpp
+++ b/Engine/Plugins/Experimental/Enterprise/DatasmithRuntime/Source/Private/TextureImporter.cpp
@@ -227,6 +227,9 @@ namespace DatasmithRuntime
 
 		if (Texture)
 		{
+			// Apply metadata on newly created texture if any
+			ApplyMetadata(AssetData.MetadataId, Texture);
+
 			uint32 TextureHash = GetTypeHash(TextureElement->CalculateElementHash(true), EDataType::Texture);
 
 			FAssetRegistry::UnregisteredAssetsData(THelper, SceneKey, [this, &Texture, TextureHash](FAssetData& AssetData) -> void
diff --git a/Engine/Plugins/Experimental/Enterprise/DatasmithRuntime/Source/Public/DatasmithRuntime.h b/Engine/Plugins/Experimental/Enterprise/DatasmithRuntime/Source/Public/DatasmithRuntime.h
--- a/Engine/Plugins/Experimental/Enterprise/DatasmithRuntime/Source/Public/DatasmithRuntime.h
+++ b/Engine/Plugins/Experimental/Enterprise/DatasmithRuntime/Source/Public/DatasmithRuntime.h
@@ -136,17 +136,10 @@ struct FDatasmithRuntimeImportOptions
 	TEnumAsByte<ECollisionTraceFlag> CollisionType = ECollisionTraceFlag::CTF_UseComplexAsSimple;
 	/**
 	 * Indicates whether meta-data should be imported or not
-	 * Meta-data are not imported by default
+	 * Meta-data are imported by default
 	 */
 	UPROPERTY(Category = "DatasmithRuntime", EditDefaultsOnly, BlueprintReadWrite)
-	bool bImportMetaData = false;
-
-	/**
-	* Indicates whether the viewpoint should be modified to match the one of the first imported camera element
-	* By default, the viewpoint is not modified
-	*/
-	UPROPERTY(Category = "DatasmithRuntime", EditDefaultsOnly, BlueprintReadWrite)
-	bool bModifyViewpoint = false;
+	bool bImportMetaData = true;
 };
 
 UCLASS(meta = (DisplayName = "Datasmith Destination"))
