Fixed issue with texture transformation

diff --git a/Engine/Plugins/Experimental/Enterprise/DatasmithRuntime/Source/Private/MaterialImportUtils.cpp b/Engine/Plugins/Experimental/Enterprise/DatasmithRuntime/Source/Private/MaterialImportUtils.cpp
--- a/Engine/Plugins/Experimental/Enterprise/DatasmithRuntime/Source/Private/MaterialImportUtils.cpp
+++ b/Engine/Plugins/Experimental/Enterprise/DatasmithRuntime/Source/Private/MaterialImportUtils.cpp
@@ -282,7 +282,7 @@ namespace DatasmithRuntime
 			int32           CoordinateIndex = 0;
 			bool            bMirrorU = false;
 			bool            bMirrorV = false;
-			float           Fading = 0.0f;
+			float           Fading = 1.0f;
 			float           UTiling = 1.0f;
 			float           VTiling = 1.0f;
 			float           UTilingPivot = 0.5f;
@@ -408,13 +408,16 @@ namespace DatasmithRuntime
 					}
 
 					FString RootName(PbrTexturePropertyNames[SlotIndex].ToString());
-					MaterialInstance->SetScalarParameterValue(*(RootName + TEXT("Fading")), 1.f /*Texture->Fading*/);
-					MaterialInstance->SetScalarParameterValue(*(RootName + TEXT("_UVOffsetX")), Texture->UOffset);
-					MaterialInstance->SetScalarParameterValue(*(RootName + TEXT("_UVOffsetY")), Texture->VOffset);
-					//MaterialInstance->SetScalarParameterValue(*(RootName + TEXT("_UVScaleX")), Texture->UOffset);
-					//MaterialInstance->SetScalarParameterValue(*(RootName + TEXT("_UVScaleY")), Texture->UOffset);
-					MaterialInstance->SetScalarParameterValue(*(RootName + TEXT("_UVWAngle")), Texture->Rotation);
-					MaterialInstance->SetScalarParameterValue(*(RootName + TEXT("_UVOffsetX")), Texture->UVOffset);
+					MaterialInstance->SetScalarParameterValue(*(RootName + TEXT("Fading")), Texture->Fading);
+					MaterialInstance->SetScalarParameterValue(*(RootName + TEXT("_UOffset")), Texture->UOffset);
+					MaterialInstance->SetScalarParameterValue(*(RootName + TEXT("_VOffset")), Texture->VOffset);
+					MaterialInstance->SetScalarParameterValue(*(RootName + TEXT("_UTiling")), Texture->UTiling);
+					MaterialInstance->SetScalarParameterValue(*(RootName + TEXT("_VTiling")), Texture->VTiling);
+					MaterialInstance->SetScalarParameterValue(*(RootName + TEXT("_UTilingPivot")), Texture->UTilingPivot);
+					MaterialInstance->SetScalarParameterValue(*(RootName + TEXT("_VTilingPivot")), Texture->VTilingPivot);
+					MaterialInstance->SetScalarParameterValue(*(RootName + TEXT("_RotAngle")), Texture->Rotation);
+					MaterialInstance->SetScalarParameterValue(*(RootName + TEXT("_URotPivot")), Texture->URotationPivot);
+					MaterialInstance->SetScalarParameterValue(*(RootName + TEXT("_VRotPivot")), Texture->VRotationPivot);
 				}
 			};
 
