Fixed crash when assets'names are not unique across runtime actors.
Removed useless ensure
Forced the use of full precision floats for UVs in the vertex buffer.

diff --git a/Engine/Plugins/Experimental/Enterprise/DatasmithRuntime/Source/Private/MaterialImporter.cpp b/Engine/Plugins/Experimental/Enterprise/DatasmithRuntime/Source/Private/MaterialImporter.cpp
--- a/Engine/Plugins/Experimental/Enterprise/DatasmithRuntime/Source/Private/MaterialImporter.cpp
+++ b/Engine/Plugins/Experimental/Enterprise/DatasmithRuntime/Source/Private/MaterialImporter.cpp
@@ -57,17 +57,19 @@ namespace DatasmithRuntime
 			}
 			else
 			{
-				FString MaterialName = FString::Printf(TEXT("M_%s_%d"), Element->GetName(), MaterialData.ElementId);
-
 #ifdef ASSET_DEBUG
 				UPackage* Package = CreatePackage(*FPaths::Combine( TEXT("/Game/Runtime/Materials"), MaterialName));
-				MaterialData.Object = TWeakObjectPtr<UObject>( UMaterialInstanceDynamic::Create( nullptr, Package, *MaterialName) );
+				MaterialData.Object = TWeakObjectPtr<UObject>( UMaterialInstanceDynamic::Create( nullptr, Package) );
 				MaterialData.Object->SetFlags(RF_Public);
 #else
-				MaterialData.Object = TWeakObjectPtr<UObject>( UMaterialInstanceDynamic::Create( nullptr, nullptr, *MaterialName) );
+				MaterialData.Object = TWeakObjectPtr<UObject>( UMaterialInstanceDynamic::Create( nullptr, nullptr) );
 #endif
 				check(MaterialData.Object.IsValid());
 
+				FString MaterialName = FString::Printf(TEXT("M_%s_%d"), Element->GetName(), MaterialData.ElementId);
+				MaterialName = FDatasmithUtils::SanitizeObjectName(MaterialName);
+				RenameObject(MaterialData.GetObject(), *MaterialName);
+
 				// Load metadata on newly created material asset if any
 				ApplyMetadata(MaterialData.MetadataId, MaterialData.GetObject());
 			}
diff --git a/Engine/Plugins/Experimental/Enterprise/DatasmithRuntime/Source/Private/MeshImporter.cpp b/Engine/Plugins/Experimental/Enterprise/DatasmithRuntime/Source/Private/MeshImporter.cpp
--- a/Engine/Plugins/Experimental/Enterprise/DatasmithRuntime/Source/Private/MeshImporter.cpp
+++ b/Engine/Plugins/Experimental/Enterprise/DatasmithRuntime/Source/Private/MeshImporter.cpp
@@ -81,15 +81,17 @@ namespace DatasmithRuntime
 			else
 			{
 				FString MeshName = FString::Printf(TEXT("SM_%s_%d"), MeshElement->GetName(), MeshData.ElementId);
-#ifdef ASSET_DEBUG
 				MeshName = FDatasmithUtils::SanitizeObjectName(MeshName);
+#ifdef ASSET_DEBUG
 				UPackage* Package = CreatePackage(*FPaths::Combine( TEXT("/Game/Runtime/Meshes"), MeshName));
-				StaticMesh = NewObject< URuntimeMesh >(Package, *MeshName, RF_Public);
+				StaticMesh = NewObject< URuntimeMesh >(Package, NAME_None, RF_Public);
 #else
-				StaticMesh = NewObject< URuntimeMesh >(GetTransientPackage(), *MeshName);
+				StaticMesh = NewObject< URuntimeMesh >(GetTransientPackage());
 #endif
 				check(StaticMesh);
 
+				RenameObject(StaticMesh, *MeshName);
+
 				StaticMesh->SetWorld(RootComponent->GetWorld());
 
 				// Add the creation of the mesh to the queue
@@ -406,6 +408,9 @@ namespace DatasmithRuntime
 			}
 		}
 
+		// Force the generation of UVs data with full precision in the vertex buffer
+		StaticMesh->GetSourceModel(0).BuildSettings.bUseFullPrecisionUVs = true;
+
 		TArray<const FMeshDescription*> MeshDescriptionPointers;
 		for (FMeshDescription& MeshDescription : MeshDescriptions)
 		{
@@ -436,6 +441,7 @@ namespace DatasmithRuntime
 
 		// Free up memory
 		MeshDescriptions.Empty();
+
 #if WITH_EDITORONLY_DATA
 		StaticMesh->ClearMeshDescriptions();
 #endif
diff --git a/Engine/Plugins/Experimental/Enterprise/DatasmithRuntime/Source/Private/SceneImporter.cpp b/Engine/Plugins/Experimental/Enterprise/DatasmithRuntime/Source/Private/SceneImporter.cpp
--- a/Engine/Plugins/Experimental/Enterprise/DatasmithRuntime/Source/Private/SceneImporter.cpp
+++ b/Engine/Plugins/Experimental/Enterprise/DatasmithRuntime/Source/Private/SceneImporter.cpp
@@ -1232,7 +1232,6 @@ namespace DatasmithRuntime
 		if (DataType == EDataType::Texture)
 		{
 			AssetPrefixedName = TexturePrefix + Elements[ElementId]->GetName();
-			ensure(AssetData.Object.IsValid() || TextureDataList.Remove(ElementId) > 0);
 		}
 		else if (DataType == EDataType::Material || DataType == EDataType::PbrMaterial)
 		{
