Specified list of allowed platforms to prevent build break on games platforms

diff --git a/Engine/Plugins/Experimental/Enterprise/DatasmithRuntime/DatasmithRuntime.uplugin b/Engine/Plugins/Experimental/Enterprise/DatasmithRuntime/DatasmithRuntime.uplugin
--- a/Engine/Plugins/Experimental/Enterprise/DatasmithRuntime/DatasmithRuntime.uplugin
+++ b/Engine/Plugins/Experimental/Enterprise/DatasmithRuntime/DatasmithRuntime.uplugin
@@ -17,7 +17,8 @@
 		{
 			"Name": "DatasmithRuntime",
 			"Type": "Runtime",
-			"LoadingPhase": "PostDefault"
+			"LoadingPhase": "PostDefault",
+			"PlatformAllowList": [ "Win64", "Mac", "Linux" ]
 		}
 	],
 	"Plugins": [
