Fixed build break when building a game

diff --git a/Engine/Plugins/Experimental/Enterprise/DatasmithRuntime/Source/Private/MeshImporter.cpp b/Engine/Plugins/Experimental/Enterprise/DatasmithRuntime/Source/Private/MeshImporter.cpp
--- a/Engine/Plugins/Experimental/Enterprise/DatasmithRuntime/Source/Private/MeshImporter.cpp
+++ b/Engine/Plugins/Experimental/Enterprise/DatasmithRuntime/Source/Private/MeshImporter.cpp
@@ -408,8 +408,10 @@ namespace DatasmithRuntime
 			}
 		}
 
+#if WITH_EDITOR
 		// Force the generation of UVs data with full precision in the vertex buffer
 		StaticMesh->GetSourceModel(0).BuildSettings.bUseFullPrecisionUVs = true;
+#endif
 
 		TArray<const FMeshDescription*> MeshDescriptionPointers;
 		for (FMeshDescription& MeshDescription : MeshDescriptions)
