Fixed flawed logic when handling tasks to be retried
Took the opportunity to fix unreported flawed logic on concurrent imports

diff --git a/Engine/Plugins/Experimental/Enterprise/DatasmithRuntime/Source/Private/DatasmithRuntime.cpp b/Engine/Plugins/Experimental/Enterprise/DatasmithRuntime/Source/Private/DatasmithRuntime.cpp
--- a/Engine/Plugins/Experimental/Enterprise/DatasmithRuntime/Source/Private/DatasmithRuntime.cpp
+++ b/Engine/Plugins/Experimental/Enterprise/DatasmithRuntime/Source/Private/DatasmithRuntime.cpp
@@ -67,11 +67,10 @@ void ADatasmithRuntimeActor::Tick(float DeltaTime)
 	if (SceneElement.IsValid() && bReceivingStarted && bReceivingEnded)
 	{
 		UE_LOG(LogDatasmithRuntime, Log, TEXT("ADatasmithRuntimeActor::Tick - Process scene's changes"));
-		if (!bImportingScene)
-		{
-			// Prevent any other DatasmithRuntime actors to import concurrently
-			bImportingScene = true;
 
+		// Prevent any other DatasmithRuntime actors to import concurrently
+		if (!bImportingScene.exchange(true))
+		{
 			if (Translator.IsValid())
 			{
 				SceneImporter->SetTranslator(Translator);
@@ -131,16 +130,19 @@ void ADatasmithRuntimeActor::EndPlay(const EEndPlayReason::Type EndPlayReason)
 
 void ADatasmithRuntimeActor::OnOpenDelta(/*int32 ElementsCount*/)
 {
-	// Block the DirectLink thread, if we are still processing the previous delta
-	while (bReceivingStarted || bImportingScene)
+	if (!IsInGameThread())
 	{
-		FPlatformProcess::SleepNoStats(0.1f);
+		// Block the DirectLink thread, if we are still processing the previous delta
+		while (bReceivingStarted)
+		{
+			FPlatformProcess::SleepNoStats(0.1f);
+		}
 	}
 
 	UE_LOG(LogDatasmithRuntime, Log, TEXT("ADatasmithRuntimeActor::OnOpenDelta"));
 	bNewScene = false;
-	bReceivingStarted = DirectLinkHelper.IsValid();
-	if (bReceivingStarted)
+	bReceivingStarted = Translator.IsValid() || DirectLinkHelper.IsValid();
+	if (DirectLinkHelper.IsValid())
 	{
 		SceneElement = DirectLinkHelper->GetScene();
 	}
@@ -246,7 +248,7 @@ void ADatasmithRuntimeActor::OnCloseDelta()
 		return;
 	}
 
-	bReceivingEnded = DirectLinkHelper.IsValid();
+	bReceivingEnded = DirectLinkHelper.IsValid() || Translator.IsValid();
 }
 
 void ADatasmithRuntimeActor::ApplyNewScene()
@@ -309,7 +311,7 @@ bool ADatasmithRuntimeActor::LoadFile(const FString& FilePath)
 
 	// Wait for any ongoing import to complete
 	// #ue_datasmithruntime: To do add code to interrupt 
-	while (bReceivingStarted || bImportingScene)
+	while (bReceivingStarted)
 	{
 		FPlatformProcess::SleepNoStats(0.1f);
 	}
@@ -376,7 +378,7 @@ namespace DatasmithRuntime
 			return false;
 		}
 
-		while(RuntimeActor->IsReceiving())
+		while(RuntimeActor->IsReceiving() || ADatasmithRuntimeActor::bImportingScene)
 		{
 			ThreadEvent->Wait(FTimespan::FromMilliseconds(50));
 		}
diff --git a/Engine/Plugins/Experimental/Enterprise/DatasmithRuntime/Source/Private/SceneImporter.cpp b/Engine/Plugins/Experimental/Enterprise/DatasmithRuntime/Source/Private/SceneImporter.cpp
--- a/Engine/Plugins/Experimental/Enterprise/DatasmithRuntime/Source/Private/SceneImporter.cpp
+++ b/Engine/Plugins/Experimental/Enterprise/DatasmithRuntime/Source/Private/SceneImporter.cpp
@@ -552,18 +552,11 @@ namespace DatasmithRuntime
 				}
 
 				const FSceneGraphId ElementId = ActionTask.GetElementId();
-				if (DirectLink::InvalidId == ElementId)
+				FBaseData& ElementData = DirectLink::InvalidId == ElementId ? FAssetData::EmptyAsset : (AssetDataList.Contains(ElementId) ? (FBaseData&)AssetDataList[ElementId] : (FBaseData&)ActorDataList[ElementId]);
+				if (ActionTask.Execute(ElementData) == EActionResult::Retry)
 				{
-					ActionTask.Execute(FAssetData::EmptyAsset);
-				}
-				else
-				{
-					FBaseData& ElementData = AssetDataList.Contains(ElementId) ? (FBaseData&)AssetDataList[ElementId] : (FBaseData&)ActorDataList[ElementId];
-					if (ActionTask.Execute(ElementData) == EActionResult::Retry)
-					{
-						ActionQueues[EQueueTask::NonAsyncQueue].Enqueue(MoveTemp(ActionTask));
-						continue;
-					}
+					ActionQueues[EQueueTask::NonAsyncQueue].Enqueue(MoveTemp(ActionTask));
+					continue;
 				}
 			}
 		}
