DatasmithRuntime - Fixing false positive ensure when failing to  delete skipped or uncompleted elements.

diff --git a/Engine/Plugins/Experimental/Enterprise/DatasmithRuntime/Source/Private/SceneImporter.cpp b/Engine/Plugins/Experimental/Enterprise/DatasmithRuntime/Source/Private/SceneImporter.cpp
--- a/Engine/Plugins/Experimental/Enterprise/DatasmithRuntime/Source/Private/SceneImporter.cpp
+++ b/Engine/Plugins/Experimental/Enterprise/DatasmithRuntime/Source/Private/SceneImporter.cpp
@@ -1127,13 +1127,13 @@ namespace DatasmithRuntime
 		// Mark assets which are about to be deleted with 'PendingDelete'
 		for (DirectLink::FSceneGraphId& ElementId : UpdateContext.Deletions)
 		{
-			if (AssetDataList.Contains(ElementId))
+			if (FAssetData* AssetData = AssetDataList.Find(ElementId))
 			{
-				AssetDataList[ElementId].SetState(EAssetState::PendingDelete);
+				AssetData->AddState(EAssetState::PendingDelete);
 			}
-			else if (ActorDataList.Contains(ElementId))
+			else if (FActorData* ActorData = ActorDataList.Find(ElementId))
 			{
-				ActorDataList[ElementId].SetState(EAssetState::PendingDelete);
+				ActorData->AddState(EAssetState::PendingDelete);
 			}
 		}
 
@@ -1184,31 +1184,23 @@ namespace DatasmithRuntime
 	EActionResult::Type FSceneImporter::DeleteElement(FSceneGraphId ElementId)
 	{
 		bool bDeletionSuccessful = false;
+		bool bHasSomethingToDelete = false;
 
-		if (AssetDataList.Contains(ElementId))
+		FAssetData AssetData(DirectLink::InvalidId);
+		if (AssetDataList.RemoveAndCopyValue(ElementId, AssetData))
 		{
-			FAssetData AssetData(DirectLink::InvalidId);
-			if (!AssetDataList.RemoveAndCopyValue(ElementId, AssetData))
-			{
-				ensure(false);
-				return EActionResult::Failed;
-			}
-
+			bHasSomethingToDelete = AssetData.HasState(EAssetState::Completed) && !AssetData.HasState(EAssetState::Skipped);
 			bDeletionSuccessful = DeleteAsset(AssetData);
 		}
 
-		if (ActorDataList.Contains(ElementId))
+		FActorData ActorData(DirectLink::InvalidId);
+		if (ActorDataList.RemoveAndCopyValue(ElementId, ActorData))
 		{
-			FActorData ActorData(DirectLink::InvalidId);
-			if (!ActorDataList.RemoveAndCopyValue(ElementId, ActorData))
-			{
-				return EActionResult::Failed;
-			}
-
+			bHasSomethingToDelete = ActorData.HasState(EAssetState::Completed) && !ActorData.HasState(EAssetState::Skipped);
 			bDeletionSuccessful = DeleteComponent(ActorData);
 		}
 
-		ensure(bDeletionSuccessful);
+		ensure(bDeletionSuccessful || !bHasSomethingToDelete);
 
 		TSharedPtr<IDatasmithElement> ElementPtr;
 
