Allow calls to CreateChildActor to supply a function to manipulate properties before the construction script is run.
Modify LightImporter to use child actor customizer instead of rerun construction scripts

diff --git a/Engine/Plugins/Experimental/Enterprise/DatasmithRuntime/Source/Private/LightImporter.cpp b/Engine/Plugins/Experimental/Enterprise/DatasmithRuntime/Source/Private/LightImporter.cpp
--- a/Engine/Plugins/Experimental/Enterprise/DatasmithRuntime/Source/Private/LightImporter.cpp
+++ b/Engine/Plugins/Experimental/Enterprise/DatasmithRuntime/Source/Private/LightImporter.cpp
@@ -309,58 +309,58 @@ namespace DatasmithRuntime
 				ChildActorComponent = CreateComponent< UChildActorComponent >(ActorData, ParentActor);
 
 				ChildActorComponent->SetChildActorClass( TSubclassOf< AActor > ( LightShapeBlueprint->GeneratedClass ) );
-				ChildActorComponent->CreateChildActor();
 			}
 
-			ADatasmithAreaLightActor* LightShapeActor = Cast< ADatasmithAreaLightActor >( ChildActorComponent->GetChildActor() );
+			ChildActorComponent->DestroyChildActor();
 
-			if ( LightShapeActor )
+			auto LightShapeCustomizer = [AreaLightElement](AActor* Actor)
 			{
-				RenameObject(LightShapeActor, AreaLightElement->GetName());
-#if WITH_EDITOR
-				LightShapeActor->SetActorLabel( AreaLightElement->GetLabel() );
-#endif
-
-				LightShapeActor->UnregisterAllComponents( true );
-
-				LightShapeActor->LightType = GetLightActorTypeForLightType( AreaLightElement->GetLightType() );
-				LightShapeActor->LightShape = (EDatasmithAreaLightActorShape)AreaLightElement->GetLightShape();
-				LightShapeActor->Dimensions = FVector2D( AreaLightElement->GetLength(), AreaLightElement->GetWidth() );
-				LightShapeActor->Color = AreaLightElement->GetColor();
-				LightShapeActor->Intensity = AreaLightElement->GetIntensity();
-				LightShapeActor->IntensityUnits = (ELightUnits)AreaLightElement->GetIntensityUnits();
-
-				if ( AreaLightElement->GetUseTemperature() )
+				if (ADatasmithAreaLightActor* LightShapeActor = Cast<ADatasmithAreaLightActor>(Actor))
 				{
-					LightShapeActor->Temperature = AreaLightElement->GetTemperature();
+					LightShapeActor->LightType = GetLightActorTypeForLightType( AreaLightElement->GetLightType() );
+					LightShapeActor->LightShape = (EDatasmithAreaLightActorShape)AreaLightElement->GetLightShape();
+					LightShapeActor->Dimensions = FVector2D( AreaLightElement->GetLength(), AreaLightElement->GetWidth() );
+					LightShapeActor->Color = AreaLightElement->GetColor();
+					LightShapeActor->Intensity = AreaLightElement->GetIntensity();
+					LightShapeActor->IntensityUnits = (ELightUnits)AreaLightElement->GetIntensityUnits();
+
+					if ( AreaLightElement->GetUseTemperature() )
+					{
+						LightShapeActor->Temperature = AreaLightElement->GetTemperature();
+					}
+
+					if ( AreaLightElement->GetUseIes() )
+					{
+						LightShapeActor->bUseIESBrightness = AreaLightElement->GetUseIesBrightness();
+						LightShapeActor->IESBrightnessScale = AreaLightElement->GetIesBrightnessScale();
+						LightShapeActor->Rotation = AreaLightElement->GetIesRotation().Rotator();
+					}
+
+					if ( AreaLightElement->GetSourceRadius() > 0.f )
+					{
+						LightShapeActor->SourceRadius = AreaLightElement->GetSourceRadius();
+					}
+
+					if ( AreaLightElement->GetSourceLength() > 0.f )
+					{
+						LightShapeActor->SourceLength = AreaLightElement->GetSourceLength();
+					}
+
+					if ( AreaLightElement->GetAttenuationRadius() > 0.f )
+					{
+						LightShapeActor->AttenuationRadius = AreaLightElement->GetAttenuationRadius();
+					}
 				}
+			};
 
-				if ( AreaLightElement->GetUseIes() )
-				{
-					LightShapeActor->bUseIESBrightness = AreaLightElement->GetUseIesBrightness();
-					LightShapeActor->IESBrightnessScale = AreaLightElement->GetIesBrightnessScale();
-					LightShapeActor->Rotation = AreaLightElement->GetIesRotation().Rotator();
-				}
-
-				if ( AreaLightElement->GetSourceRadius() > 0.f )
-				{
-					LightShapeActor->SourceRadius = AreaLightElement->GetSourceRadius();
-				}
-
-				if ( AreaLightElement->GetSourceLength() > 0.f )
-				{
-					LightShapeActor->SourceLength = AreaLightElement->GetSourceLength();
-				}
-
-				if ( AreaLightElement->GetAttenuationRadius() > 0.f )
-				{
-					LightShapeActor->AttenuationRadius = AreaLightElement->GetAttenuationRadius();
-				}
-
-				LightShapeActor->RegisterAllComponents();
-
-				LightShapeActor->RerunConstructionScripts();
+			ChildActorComponent->CreateChildActor(MoveTemp(LightShapeCustomizer));
 
+			if (ADatasmithAreaLightActor* LightShapeActor = Cast<ADatasmithAreaLightActor>(ChildActorComponent->GetChildActor()))
+			{
+				RenameObject(LightShapeActor, AreaLightElement->GetName());
+#if WITH_EDITOR
+				LightShapeActor->SetActorLabel(AreaLightElement->GetLabel());
+#endif
 				return ChildActorComponent;
 			}
 		}
diff --git a/Engine/Source/Runtime/Engine/Classes/Components/ChildActorComponent.h b/Engine/Source/Runtime/Engine/Classes/Components/ChildActorComponent.h
--- a/Engine/Source/Runtime/Engine/Classes/Components/ChildActorComponent.h
+++ b/Engine/Source/Runtime/Engine/Classes/Components/ChildActorComponent.h
@@ -174,7 +174,7 @@ public:
 	void ApplyComponentInstanceData(FChildActorComponentInstanceData* ComponentInstanceData, const ECacheApplyPhase CacheApplyPhase);
 
 	/** Create the child actor */
-	virtual void CreateChildActor();
+	virtual void CreateChildActor(TFunction<void(AActor*)> CustomizerFunc = nullptr);
 
 	AActor* GetChildActor() const { return ChildActor; }
 	AActor* GetChildActorTemplate() const { return ChildActorTemplate; }
diff --git a/Engine/Source/Runtime/Engine/Private/Components/ChildActorComponent.cpp b/Engine/Source/Runtime/Engine/Private/Components/ChildActorComponent.cpp
--- a/Engine/Source/Runtime/Engine/Private/Components/ChildActorComponent.cpp
+++ b/Engine/Source/Runtime/Engine/Private/Components/ChildActorComponent.cpp
@@ -592,7 +592,7 @@ bool UChildActorComponent::IsChildActorReplicated() const
 	return (bChildActorClassReplicated || bChildActorTemplateReplicated);
 }
 
-void UChildActorComponent::CreateChildActor()
+void UChildActorComponent::CreateChildActor(TFunction<void(AActor*)> CustomizerFunc)
 {
 	AActor* MyOwner = GetOwner();
 
@@ -681,6 +681,11 @@ void UChildActorComponent::CreateChildActor()
 
 					ChildActorName = ChildActor->GetFName();
 
+					if (CustomizerFunc)
+					{
+						CustomizerFunc(ChildActor);
+					}
+
 					// Parts that we deferred from SpawnActor
 					const FComponentInstanceDataCache* ComponentInstanceData = (CachedInstanceData ? CachedInstanceData->ComponentInstanceData.Get() : nullptr);
 					ChildActor->FinishSpawning(GetComponentTransform(), false, ComponentInstanceData);
