Removing dependency on deleted plugin RuntimePhysXCooking

diff --git a/Engine/Plugins/Experimental/Enterprise/DatasmithRuntime/DatasmithRuntime.uplugin b/Engine/Plugins/Experimental/Enterprise/DatasmithRuntime/DatasmithRuntime.uplugin
--- a/Engine/Plugins/Experimental/Enterprise/DatasmithRuntime/DatasmithRuntime.uplugin
+++ b/Engine/Plugins/Experimental/Enterprise/DatasmithRuntime/DatasmithRuntime.uplugin
@@ -26,10 +26,6 @@
 			"Name": "DatasmithImporter",
 			"Enabled": true
 		},
-		{
-			"Name": "RuntimePhysXCooking",
-			"Enabled": true
-		},
 		{
 			"Name": "UdpMessaging",
 			"Enabled": true
