Fixed Race condition in ADatasmithRuntimeActor::LoadFile
 - Solution: Revisit use of ThreadEvent which was wrongly used as a timer.

diff --git a/Engine/Plugins/Experimental/Enterprise/DatasmithRuntime/Source/Private/DatasmithRuntime.cpp b/Engine/Plugins/Experimental/Enterprise/DatasmithRuntime/Source/Private/DatasmithRuntime.cpp
--- a/Engine/Plugins/Experimental/Enterprise/DatasmithRuntime/Source/Private/DatasmithRuntime.cpp
+++ b/Engine/Plugins/Experimental/Enterprise/DatasmithRuntime/Source/Private/DatasmithRuntime.cpp
@@ -354,7 +354,7 @@ namespace DatasmithRuntime
 {
 	bool FTranslationJob::Execute()
 	{
-		if (!RuntimeActor.IsValid()|| ThreadEvent == nullptr)
+		if (!RuntimeActor.IsValid())
 		{
 			return false;
 		}
@@ -380,7 +380,7 @@ namespace DatasmithRuntime
 
 		while(RuntimeActor->IsReceiving() || ADatasmithRuntimeActor::bImportingScene)
 		{
-			ThreadEvent->Wait(FTimespan::FromMilliseconds(50));
+			FPlatformProcess::Sleep(0.05f);
 		}
 
 		RuntimeActor->OnOpenDelta();
@@ -422,8 +422,11 @@ namespace DatasmithRuntime
 				continue;
 			}
 
-			ThreadEvent->Wait(FTimespan::FromMilliseconds(50));
+			FPlatformProcess::Sleep(0.1f);
 		}
+
+		// The FTranslationThread is being deleted, trigger the event before exiting
+		ThreadEvent->Trigger();
 	}
 
 	FTranslationThread::~FTranslationThread()
@@ -431,8 +434,9 @@ namespace DatasmithRuntime
 		if (bKeepRunning)
 		{
 			bKeepRunning = false;
-			ThreadEvent->Trigger();
-			ThreadResult.Get();
+			ThreadResult.Reset();
+			// Wait for FTranslationThread::Run() to be completed
+			ThreadEvent->Wait();
 			FPlatformProcess::ReturnSynchEventToPool(ThreadEvent);
 		}
 	}
diff --git a/Engine/Plugins/Experimental/Enterprise/DatasmithRuntime/Source/Public/DatasmithRuntime.h b/Engine/Plugins/Experimental/Enterprise/DatasmithRuntime/Source/Public/DatasmithRuntime.h
--- a/Engine/Plugins/Experimental/Enterprise/DatasmithRuntime/Source/Public/DatasmithRuntime.h
+++ b/Engine/Plugins/Experimental/Enterprise/DatasmithRuntime/Source/Public/DatasmithRuntime.h
@@ -47,22 +47,18 @@ namespace DatasmithRuntime
 		FTranslationJob(ADatasmithRuntimeActor* InActor, const FString& InFilePath)
 			: RuntimeActor(InActor)
 			, FilePath(InFilePath)
-			, ThreadEvent(nullptr)
 		{
 		}
 
-		FTranslationJob() : ThreadEvent(nullptr)
+		FTranslationJob()
 		{
 		}
 
 		bool Execute();
 
-		void SetEvent(FEvent* InThreadEvent) { ThreadEvent = InThreadEvent; }
-
 	private:
 		TWeakObjectPtr<ADatasmithRuntimeActor> RuntimeActor;
 		FString FilePath;
-		FEvent* ThreadEvent;
 	};
 
 	class FTranslationThread
@@ -70,7 +66,6 @@ namespace DatasmithRuntime
 	public:
 		FTranslationThread() 
 			: bKeepRunning(false)
-			, ThreadEvent(nullptr)
 		{}
 
 		~FTranslationThread();
@@ -79,13 +74,12 @@ namespace DatasmithRuntime
 
 		void AddJob(FTranslationJob&& Job)
 		{
-			Job.SetEvent(ThreadEvent);
 			JobQueue.Enqueue(MoveTemp(Job));
 		}
 
 		std::atomic_bool bKeepRunning;
 		TFuture<void> ThreadResult;
-		FEvent* ThreadEvent;
+		FEvent* ThreadEvent = nullptr;
 		TQueue< FTranslationJob, EQueueMode::Mpsc > JobQueue;
 	};
 }
