Fixed compilation error for SceneImporter.cpp in NonUnity  mode

diff --git a/Engine/Plugins/Experimental/Enterprise/DatasmithRuntime/Source/Private/SceneImporter.cpp b/Engine/Plugins/Experimental/Enterprise/DatasmithRuntime/Source/Private/SceneImporter.cpp
--- a/Engine/Plugins/Experimental/Enterprise/DatasmithRuntime/Source/Private/SceneImporter.cpp
+++ b/Engine/Plugins/Experimental/Enterprise/DatasmithRuntime/Source/Private/SceneImporter.cpp
@@ -15,6 +15,7 @@
 #include "Materials/MaterialInstanceDynamic.h"
 #include "ProfilingDebugging/MiscTrace.h"
 #include "RenderingThread.h"
+#include "Misc/Paths.h"
 
 namespace DatasmithRuntime
 {
