Remove unused variables

diff --git a/Engine/Plugins/Experimental/Enterprise/DatasmithRuntime/Source/Public/DatasmithRuntime.h b/Engine/Plugins/Experimental/Enterprise/DatasmithRuntime/Source/Public/DatasmithRuntime.h
--- a/Engine/Plugins/Experimental/Enterprise/DatasmithRuntime/Source/Public/DatasmithRuntime.h
+++ b/Engine/Plugins/Experimental/Enterprise/DatasmithRuntime/Source/Public/DatasmithRuntime.h
@@ -226,11 +226,6 @@ private:
 	static std::atomic_bool bImportingScene;
 	FUpdateContext UpdateContext;
 
-#if WITH_EDITOR
-	int32 EnableThreadedImport = MAX_int32;
-	int32 EnableCADCache = MAX_int32;
-#endif
-
 	static TUniquePtr<DatasmithRuntime::FTranslationThread> TranslationThread;
 
 	friend class DatasmithRuntime::FTranslationJob;
