Fixed mismatched copy constructors/assignment operators for DatasmithRuntime::FBaseData, FAssetData and FActorData.

diff --git a/Engine/Plugins/Experimental/Enterprise/DatasmithRuntime/Source/Private/SceneImporter.h b/Engine/Plugins/Experimental/Enterprise/DatasmithRuntime/Source/Private/SceneImporter.h
--- a/Engine/Plugins/Experimental/Enterprise/DatasmithRuntime/Source/Private/SceneImporter.h
+++ b/Engine/Plugins/Experimental/Enterprise/DatasmithRuntime/Source/Private/SceneImporter.h
@@ -216,6 +216,9 @@ namespace DatasmithRuntime
 			Referencers = Other.Referencers;
 		}
 
+		// Non-assignable due to DataState atomic member
+		FBaseData& operator=(const FBaseData& Other) = delete;
+
 		bool HasState(EAssetState Value) const
 		{
 			return !!(DataState & Value);
@@ -263,6 +266,9 @@ namespace DatasmithRuntime
 		{
 		}
 
+		// Copy constructor intentionally copies the MetadataId member whereas the assignment operator doesn't - is this intentional?
+		FAssetData(const FAssetData&) = default;
+
 		FAssetData& operator=(const FAssetData& Source)
 		{
 			ElementId = Source.ElementId;
@@ -311,6 +317,9 @@ namespace DatasmithRuntime
 		{
 		}
 
+		// Copy constructor intentionally copies the MetadataId member whereas the assignment operator doesn't - is this intentional?
+		FActorData(const FActorData&) = default;
+
 		FActorData& operator=(const FActorData& Source)
 		{
 			ElementId = Source.ElementId;
