Rmove unnecessary build directive

diff --git a/Engine/Plugins/Experimental/Enterprise/DatasmithRuntime/Source/DatasmithRuntime.Build.cs b/Engine/Plugins/Experimental/Enterprise/DatasmithRuntime/Source/DatasmithRuntime.Build.cs
--- a/Engine/Plugins/Experimental/Enterprise/DatasmithRuntime/Source/DatasmithRuntime.Build.cs
+++ b/Engine/Plugins/Experimental/Enterprise/DatasmithRuntime/Source/DatasmithRuntime.Build.cs
@@ -8,7 +8,6 @@ public class DatasmithRuntime : ModuleRules
 {
 	public DatasmithRuntime(ReadOnlyTargetRules Target) : base(Target)
 	{
-		OptimizeCode = CodeOptimization.InShippingBuildsOnly;
 		PCHUsage = ModuleRules.PCHUsageMode.UseExplicitOrSharedPCHs;
 
 		PublicDependencyModuleNames.AddRange(
